WaterAdvanced: Shift plugin state from Restricted to Experimental

#virtualized

diff --git a/Engine/Plugins/Experimental/WaterAdvanced/Config/BaseWaterAdvanced.ini b/Engine/Plugins/Experimental/WaterAdvanced/Config/BaseWaterAdvanced.ini
new file mode 100644
--- /dev/null
+++ b/Engine/Plugins/Experimental/WaterAdvanced/Config/BaseWaterAdvanced.ini
@@ -0,0 +1,6 @@
+﻿[CoreRedirects]
++PropertyRedirects=(OldName="/Script/WaterAdvanced.ShallowWaterSettings.PhysicsAssetOverridesDataAsset",NewName="/Script/WaterAdvanced.ShallowWaterSettings.PhysicsAssetProxiesDataAsset")
++StructRedirects=(OldName="/Script/WaterAdvanced.ShallowWaterCollisionAgent",NewName="/Script/WaterAdvanced.ShallowWaterCollisionTracker_Actor")
++FunctionRedirects=(OldName="/Script/WaterAdvanced.ShallowWaterSubsystem.AddCollisionAgent",NewName="/Script/WaterAdvanced.ShallowWaterSubsystem.AddCollisionTrackerForActor")
++FunctionRedirects=(OldName="/Script/WaterAdvanced.ShallowWaterSubsystem.RemoveCollisionAgent",NewName="/Script/WaterAdvanced.ShallowWaterSubsystem.RemoveCollisionAgentForActor")
++FunctionRedirects=(OldName="/Script/WaterAdvanced.ShallowWaterSubsystem.RemoveCollisionAgentForActor",NewName="/Script/WaterAdvanced.ShallowWaterSubsystem.RemoveCollisionTrackerForActor")
\ No newline at end of file
diff --git a/Engine/Plugins/Experimental/WaterAdvanced/Source/WaterAdvanced/Private/ShallowWaterCommon.cpp b/Engine/Plugins/Experimental/WaterAdvanced/Source/WaterAdvanced/Private/ShallowWaterCommon.cpp
new file mode 100644
--- /dev/null
+++ b/Engine/Plugins/Experimental/WaterAdvanced/Source/WaterAdvanced/Private/ShallowWaterCommon.cpp
@@ -0,0 +1,6 @@
+﻿// Copyright Epic Games, Inc. All Rights Reserved.
+
+
+#include "ShallowWaterCommon.h"
+
+DEFINE_LOG_CATEGORY(LogShallowWater);
\ No newline at end of file
diff --git a/Engine/Plugins/Experimental/WaterAdvanced/Source/WaterAdvanced/Private/ShallowWaterSettings.cpp b/Engine/Plugins/Experimental/WaterAdvanced/Source/WaterAdvanced/Private/ShallowWaterSettings.cpp
new file mode 100644
--- /dev/null
+++ b/Engine/Plugins/Experimental/WaterAdvanced/Source/WaterAdvanced/Private/ShallowWaterSettings.cpp
@@ -0,0 +1,15 @@
+﻿// Copyright Epic Games, Inc. All Rights Reserved.
+
+#include "ShallowWaterSettings.h"
+
+UShallowWaterSettings::UShallowWaterSettings()
+{
+	DefaultShallowWaterNiagaraSimulation = TSoftObjectPtr<UNiagaraSystem>(FSoftObjectPath(TEXT("/WaterAdvanced/Niagara/Systems/Grid2D_SW_WaterBody.Grid2D_SW_WaterBody")));
+	
+	WaterMPC = TSoftObjectPtr<UMaterialParameterCollection>(FSoftObjectPath(TEXT("/Water/Materials/MPC/MPC_Water.MPC_Water")));
+}
+
+FName UShallowWaterSettings::GetCategoryName() const
+{
+	return FName("Plugins");
+}
diff --git a/Engine/Plugins/Experimental/WaterAdvanced/Source/WaterAdvanced/Private/ShallowWaterSubsystem.cpp b/Engine/Plugins/Experimental/WaterAdvanced/Source/WaterAdvanced/Private/ShallowWaterSubsystem.cpp
new file mode 100644
--- /dev/null
+++ b/Engine/Plugins/Experimental/WaterAdvanced/Source/WaterAdvanced/Private/ShallowWaterSubsystem.cpp
@@ -0,0 +1,1165 @@
+﻿// Copyright Epic Games, Inc. All Rights Reserved.
+
+#include "ShallowWaterSubsystem.h"
+
+#include "ShallowWaterSettings.h"
+#include "NiagaraComponent.h"
+#include "NiagaraDataInterfaceArrayFunctionLibrary.h"
+#include "NiagaraFunctionLibrary.h"
+#include "WaterBodyActor.h"
+#include "WaterSubsystem.h"
+#include "Components/SkeletalMeshComponent.h"
+#include "Engine/TextureRenderTarget2D.h"
+#include "Kismet/GameplayStatics.h"
+#include "Kismet/KismetMaterialLibrary.h"
+#include "DrawDebugHelpers.h"
+#include "Engine/World.h"
+#include "Engine/GameInstance.h"
+#include "Engine/LocalPlayer.h"
+#include "GameFramework/Character.h"
+#include "Materials/MaterialInstanceDynamic.h"
+#include "GameFramework/Pawn.h"
+#include "GameFramework/SpectatorPawn.h"
+#include "PhysicsEngine/PhysicsAsset.h"
+#include "Engine/AssetManager.h"
+
+//UE_DISABLE_OPTIMIZATION_SHIP
+
+int32 GSWEnabled = 1;
+static FAutoConsoleVariableRef CVarSWEnabled(
+	TEXT("r.ShallowWater.Enabled"),
+	GSWEnabled,
+	TEXT("Should create ShallowWaterSubsystem at all"),
+	ECVF_Scalability
+);
+
+float GSWRemainActiveForSeconds = 15.f;
+static FAutoConsoleVariableRef CVarSWRemainActiveForSeconds(
+	TEXT("r.ShallowWater.FadeOutWait"),
+	GSWRemainActiveForSeconds,
+	TEXT("If not any kind of collision or collision tracker is active, how long the simulation continues to be active to wait out the ripples")
+);
+
+float GSWImpactTrackerActiveForSeconds = 5.f;
+static FAutoConsoleVariableRef CVarImpactTrackerActiveForSeconds(
+	TEXT("r.ShallowWater.CollisionTracker.ImpactTrackerActiveForSeconds"),
+	GSWImpactTrackerActiveForSeconds,
+	TEXT("How long impacts (e.g. bullets hit water) are tracked, keeping the sim active") 
+);
+
+int32 GSWDrawWaterSurfaceProjection = 0;
+static FAutoConsoleVariableRef CVarSWDrawWaterSurfaceProjection(
+	TEXT("r.ShallowWater.DrawSurfaceProjection"),
+	GSWDrawWaterSurfaceProjection,
+	TEXT("")
+);
+
+FName UShallowWaterSubsystem::ColliderComponentTag = FName("RigidMesh_ShallowWaterCollider");
+
+TArray<AWaterBody*> FShallowWaterCollisionTracker_Actor::GetOverlappingWaterBodies() const
+{
+	TArray<AWaterBody*> Result;
+	if(const AActor* ActorPointer = CollisionActor.Get())
+	{
+    	TArray<AActor*> OverlappingActors;
+		ActorPointer->GetOverlappingActors(OverlappingActors, AWaterBody::StaticClass());
+		for (AActor* WaterBody : OverlappingActors)
+		{
+			if (AWaterBody* Actor = Cast<AWaterBody>(WaterBody))
+			{
+				Result.AddUnique(Actor);
+			}
+		}
+	}
+    return Result;
+}
+
+UShallowWaterSubsystem::UShallowWaterSubsystem()
+{
+}
+
+void UShallowWaterSubsystem::PostInitialize()
+{
+	Super::PostInitialize();
+
+	// Register default PA Proxies before all other GFP chimes in
+	Settings = GetMutableDefault<UShallowWaterSettings>();
+	if (Settings)
+	{
+		UAssetManager::GetStreamableManager().RequestAsyncLoad(Settings->PhysicsAssetProxiesDataAsset.ToSoftObjectPath(),
+			FStreamableDelegate::CreateWeakLambda(this, [this]()
+				{
+					if (Settings->PhysicsAssetProxiesDataAsset.IsValid())
+					{
+						RegisterPhysicsAssetProxiesDataAsset(Settings->PhysicsAssetProxiesDataAsset.Get());						
+					}
+					else
+					{
+						ensureMsgf(false, TEXT("UShallowWaterSubsystem::PostInitialize() - PhyicsAssetProxiesDataAsset is not valid"));
+					}
+				})
+		);		
+	}
+}
+
+void UShallowWaterSubsystem::OnWorldBeginPlay(UWorld& InWorld)
+{
+	TArray<ULocalPlayer*> LocalPlayers = InWorld.GetGameInstance() != nullptr ? InWorld.GetGameInstance()->GetLocalPlayers() : TArray<ULocalPlayer*>();
+	if (!LocalPlayers.IsEmpty())
+	{
+		if (ULocalPlayer* const LocalPlayer = LocalPlayers[0])
+		{
+			APlayerController* Controller = LocalPlayer->GetPlayerController(&InWorld);
+			if (Controller)
+			{
+				OnLocalPlayerControllerBecomesValid(Controller);
+			}
+			LocalPlayer->OnPlayerControllerChanged().AddUObject(this, &UShallowWaterSubsystem::OnLocalPlayerControllerBecomesValid);
+		}
+	}
+}
+
+bool UShallowWaterSubsystem::ShouldCreateSubsystem(UObject* Outer) const
+{
+	if (!GSWEnabled)
+	{
+		return false;
+	}
+	
+	if (!FApp::CanEverRender() || IsRunningDedicatedServer())
+	{
+		return false;
+	}
+
+	// IsRunningDedicatedServer() is a static check and doesn't work in PIE "As Client' mode where both a server and a client are run
+	if (UWorld* World = Cast<UWorld>(Outer))
+	{
+		if (World->IsNetMode(NM_DedicatedServer))
+		{
+			return false;
+		}
+	}
+	
+	return Super::ShouldCreateSubsystem(Outer);
+}
+
+void UShallowWaterSubsystem::Tick(float DeltaTime)
+{	
+	Super::Tick(DeltaTime);
+
+	// #todo bullet impacts should be pretty late in tick order but this can be risky
+	ResetImpacts();
+
+	if (!IsShallowWaterInitialized())
+	{
+		return;
+	}
+
+	UpdateActivePawns();
+	GatherContextsFromPawns(ActivePawns);
+	UpdateCollisionForPendingContexts();
+
+	/*
+	 **** Ticking order ****
+	 *
+	 * Niagara CPU Script and Pawns (order can be mixed)
+	 *
+	 * Subsystem
+	 *     - ThisClass::UpdateGridMovement()
+	 *     - Set SimCenter to water materials
+	 *
+	 * Single Layer Water shader
+	 */
+	// Leave the sim center at previous water body Z level untouched, if owner character is not in water
+	TSet<AWaterBody*> WaterBodies = GetAllOverlappingWaterBodiesAndUpdateCollisionTrackers();
+	if (!WaterBodies.IsEmpty())
+	{
+		UpdateOverlappingWaterBodiesHistory(WaterBodies.Array());
+	
+		if (WaterInfoTexture == nullptr)
+		{
+			TryGetOrWaitForWaterInfoTextureFromWaterBodies(WaterBodies);
+		}
+	
+		for (AWaterBody* Water : WaterBodies)
+		{
+			if (Water == nullptr)
+			{
+				continue;
+			}
+			TryUpdateWaterBodyMIDParameters(Water->GetWaterBodyComponent());
+		}
+	}
+
+	// We might miss some impacts because the Niagara System hasn't been activated in time.
+	// Note that we do this one frame after Niagara System activation otherwise it won't work
+	// due to reset tick logic in Niagara skipping the main simulation loop
+	if (bFlushPendingImpactsNextTick)
+	{		
+		FlushPendingImpacts();
+	}
+	bFlushPendingImpactsNextTick = false;
+
+	if (ShouldSimulateThisFrame())
+	{
+		UpdateGridMovement();
+		if (!ShallowWaterNiagaraSimulation->IsActive())
+		{
+			ShallowWaterNiagaraSimulation->Activate(false);
+
+			bFlushPendingImpactsNextTick = true;
+		}
+	}
+	else
+	{
+		if (ShallowWaterNiagaraSimulation->IsActive())
+		{
+			ShallowWaterNiagaraSimulation->Deactivate();
+		}
+	}
+}
+
+TStatId UShallowWaterSubsystem::GetStatId() const
+{
+	RETURN_QUICK_DECLARE_CYCLE_STAT(UShallowWaterSubsystem, STATGROUP_Tickables);
+}
+
+void UShallowWaterSubsystem::InitializeShallowWater()
+{
+	if (!ensure(WeakPlayerController.IsValid()))
+	{
+		UE_LOG(LogShallowWater, Warning, TEXT("PlayerController is invalid during initialization"));
+		return;
+	}
+
+	if(!IsShallowWaterAllowedToInitialize())
+	{
+		return;
+	}
+
+	MPC = Settings->WaterMPC.LoadSynchronous();
+	if (MPC == nullptr)
+	{
+		ensure(false);
+		UE_LOG(LogShallowWater, Warning, TEXT("UShallowWaterSubsystem::InitializeShallowWater() - MPC is invalid. Make sure it's set in ShallowWater Settings."));
+		return;
+	}
+	
+	/*
+	 * From here, non-spectator PlayerPawn might not be available if the game is a replay
+	 * So we relies on  GetTheMostRelevantPlayerPawn();
+	 */
+	const APawn* CursorPawn = GetTheMostRelevantPlayerPawn();
+	if (CursorPawn == nullptr)
+	{
+		UE_LOG(LogShallowWater, Warning, TEXT("Could not find CursorPawn during initialization"));
+		return;
+	}
+
+	UNiagaraSystem* const ShallowWaterTemplate = Settings->DefaultShallowWaterNiagaraSimulation.LoadSynchronous();
+	if (ShallowWaterTemplate == nullptr)
+	{
+		ensure(false);
+		UE_LOG(LogShallowWater, Warning, TEXT("UShallowWaterSubsystem::InitializeShallowWater() - Couldn't find ShallowWater template in settings"));
+		return;
+	}
+	const FVector SpawnLocation =  CursorPawn->GetActorLocation();
+	ShallowWaterNiagaraSimulation = UNiagaraFunctionLibrary::SpawnSystemAtLocation(GetWorld(), ShallowWaterTemplate,
+		SpawnLocation, FRotator::ZeroRotator, FVector::OneVector, false,
+		false,
+		ENCPoolMethod::None, false);
+
+	if (ShallowWaterNiagaraSimulation == nullptr)
+	{
+		UE_LOG(LogShallowWater, Warning, TEXT("UShallowWaterSubsystem::InitializeShallowWater() - ShallowWaterNiagaraSystem spawn failed"));
+		return;
+	}
+	/*
+	 * Initialize succeeds
+	 */
+
+	bIsShallowWaterInitialized = true;
+	
+	CreateRTs();
+	
+	ShallowWaterNiagaraSimulation->SetUsingAbsoluteRotation(true);
+	ShallowWaterNiagaraSimulation->Activate();
+		
+	InitializeParameters();
+
+	for (TWeakObjectPtr<AWaterBody> WeakWaterBody : PendingWaterBodiesToSetMIDOnInitialize)
+	{
+		if (AWaterBody* WaterBody = WeakWaterBody.Get())
+		{
+			TryUpdateWaterBodyMIDParameters(WaterBody->GetWaterBodyComponent());
+		}
+	}
+	PendingWaterBodiesToSetMIDOnInitialize.Empty();
+
+	PendingImpacts.Reset();
+	
+	UE_LOG(LogShallowWater, Log, TEXT("UShallowWaterSubsystem::InitializeShallowWater() finished successfully"));
+}
+
+bool UShallowWaterSubsystem::IsShallowWaterInitialized() const
+{
+	return bIsShallowWaterInitialized;
+}
+
+APawn* UShallowWaterSubsystem::GetNonSpectatorPawnFromWeakController() const
+{
+	const APlayerController* Controller = WeakPlayerController.Get();
+	if (Controller)
+	{
+		APawn* Pawn = Controller->GetPawn();
+		if (Pawn && !Pawn->GetClass()->IsChildOf(ASpectatorPawn::StaticClass()))
+		{
+			return Pawn; 
+		}
+	}
+	return nullptr;
+}
+
+TOptional<FVector> UShallowWaterSubsystem::GetCameraLocationFromWeakController() const
+{
+	TOptional<FVector> Result;
+	const APlayerController* Controller = WeakPlayerController.Get();
+	if (Controller)
+	{
+		const TObjectPtr<APlayerCameraManager> CamMan = Controller->PlayerCameraManager;
+		if (CamMan)
+		{
+			Result = CamMan->GetCameraLocation();
+			return Result;
+		}
+	}
+	return Result;
+}
+
+APawn* UShallowWaterSubsystem::GetTheMostRelevantPlayerPawn() const
+{
+	const APlayerController* Controller = WeakPlayerController.Get();
+	if (Controller)
+	{
+		APawn* Pawn = GetNonSpectatorPawnFromWeakController();
+		if (Pawn)
+		{
+			return Pawn;
+		}
+		
+		const TOptional<FVector> CamLoc = GetCameraLocationFromWeakController();
+		if (CamLoc.IsSet())
+		{
+			APawn* BestPawn = nullptr;
+			float BestDistance = FLT_MAX;
+			for (APawn* TestPawn : GetPawnsInRange())
+			{
+			    const float Distance = FVector::Distance(CamLoc.GetValue(), TestPawn->GetActorLocation()); 
+				if (Distance < BestDistance)
+				{
+					BestDistance = Distance;
+					BestPawn = TestPawn;
+				}
+			}
+			if (BestPawn)
+			{
+				return BestPawn;
+			}
+		}
+	}
+	return nullptr;
+}
+
+void UShallowWaterSubsystem::CreateRTs()
+{
+	int32 Resolution = Settings->ShallowWaterSimParameters.ResolutionMaxAxis;
+
+	if (Resolution <= 0)
+	{
+		UE_LOG(LogShallowWater, Warning, TEXT("ShallowWaterComponent: Invalid Resolution Max Axis.  Defauling to 1x1 simulation."))
+		Resolution = 1;
+	}
+
+	NormalRT = NewObject<UTextureRenderTarget2D>(this);
+	check(NormalRT);
+	NormalRT->RenderTargetFormat = ETextureRenderTargetFormat::RTF_RGB10A2;
+	NormalRT->ClearColor = FLinearColor(0.f,0.f,0.f,0.f);
+	NormalRT->bAutoGenerateMips = false;
+	NormalRT->bCanCreateUAV = true;    // Niagara RT DI requires UAV. Do we need to create it here?
+	NormalRT->InitAutoFormat(Resolution, Resolution);	
+	NormalRT->UpdateResourceImmediate(true);	
+}
+
+void UShallowWaterSubsystem::InitializeParameters()
+{
+	if (ShallowWaterNiagaraSimulation)
+	{
+		ShallowWaterNiagaraSimulation->SetVariableVec2(FName("WorldGridSize"), FVector2D(GetGridSize()));
+		ShallowWaterNiagaraSimulation->SetVariableInt(FName("ResolutionMaxAxis"), GetGridResolution());
+		ShallowWaterNiagaraSimulation->SetVariableTextureRenderTarget(FName("NormalRT"), NormalRT);
+
+		// set variables for impacts
+		ShallowWaterNiagaraSimulation->SetVariableInt(FName("MaxImpacts"), GetMaxImpactsPerFrame());
+		ShallowWaterNiagaraSimulation->SetVariableInt(FName("NumImpacts"), 0);
+		CurrentNumImpacts = 0;
+	}
+	else
+	{
+		UE_LOG(LogShallowWater, Error, TEXT("ShallowWaterComponent: No simulation found on component."))
+	}
+	
+	if (MPC)
+	{
+		UKismetMaterialLibrary::SetScalarParameterValue(this, MPC,
+			Settings->WorldGridSizeMPCName,
+			GetGridSize());
+		UKismetMaterialLibrary::SetScalarParameterValue(this, MPC,
+			Settings->ResolutionMaxAxisMPCName,
+			GetGridResolution());
+	}
+	else
+	{
+		UE_LOG(LogShallowWater, Warning, TEXT("ShallowWaterComponent: No valid MPC found in Project Settings - Water Advanced. The simulation will work but would show preview from the Niagara renderer only."))
+	}
+}
+
+void UShallowWaterSubsystem::UpdateGridMovement()
+{
+	if (ShallowWaterNiagaraSimulation == nullptr)
+	{
+		return;
+	}
+	
+	APawn* CursorPawn = GetTheMostRelevantPlayerPawn();
+	if (CursorPawn == nullptr)
+	{
+		return;
+	}
+	const FVector CursorPawnLocation = CursorPawn->GetActorLocation();
+	
+	/*
+	 * Snap the sim center to player character
+	 */
+	// #note this query gets the closest water surface in 3D space, not a 2D topdown projection as we instinctively assumed
+	// Could be better and could be worse
+
+	float BestDistanceSqr = FLT_MAX;
+	TOptional<FVector> BestWaterLocation;
+	for (TWeakObjectPtr<AWaterBody> WeakWaterBody : LastOverlappingWaterBodies_Internal)
+	{
+		if (const AWaterBody* WaterBody = WeakWaterBody.Get())
+		{
+			const FWaterBodyQueryResult WaterInfo = WaterBody->GetWaterBodyComponent()->QueryWaterInfoClosestToWorldLocation(
+				CursorPawnLocation, EWaterBodyQueryFlags::ComputeLocation);
+			const FVector WaterLocation = WaterInfo.GetWaterSurfaceLocation();
+			const float ds = FVector::DistSquared(WaterLocation, CursorPawnLocation); 
+			if (ds < BestDistanceSqr)
+			{
+				BestDistanceSqr = ds;
+				BestWaterLocation = WaterLocation;
+			}
+		}
+	}
+
+	if (!BestWaterLocation.IsSet())
+	{
+		return;
+	}
+
+	const FVector ProjectedLocation = FVector(CursorPawnLocation.X, CursorPawnLocation.Y, 0);
+	ShallowWaterNiagaraSimulation->SetWorldLocation(ProjectedLocation);
+	ShallowWaterNiagaraSimulation->SetVariableFloat(FName("WaterHeightAtParent"), BestWaterLocation.GetValue().Z);
+
+	// time left before sim is destroyed- used to attenuate waves
+	const float SecondsSinceCollision = GetWorld()->GetTimeSeconds() - LastTimeOverlappingAnyWaterBody;
+	const float SecondsUntilDestroyed = FMath::Max(0, GSWRemainActiveForSeconds - SecondsSinceCollision);
+	ShallowWaterNiagaraSimulation->SetVariableFloat(FName("SecondsUntilDestroyed"), SecondsUntilDestroyed);
+
+#if ENABLE_DRAW_DEBUG
+	if (GSWDrawWaterSurfaceProjection)
+	{
+		DrawDebugLine(GetWorld(), CursorPawnLocation, BestWaterLocation.GetValue(), FColor::Yellow, false, 0.5f);
+		DrawDebugPoint(GetWorld(), BestWaterLocation.GetValue(), 5.f, FColor::Yellow, false, 0.5f);
+	}
+#endif
+
+	/*
+	 * Feed sim center to MPC
+	 */
+	if (MPC != nullptr)
+	{
+		UKismetMaterialLibrary::SetVectorParameterValue(this, MPC, Settings->GridCenterMPCName,
+			FLinearColor(PreviousProjectedLocation.X, PreviousProjectedLocation.Y, 0.f, 0.f));
+	}
+	PreviousProjectedLocation = ProjectedLocation;
+}
+
+void UShallowWaterSubsystem::ResetImpacts()
+{
+	if (ShallowWaterNiagaraSimulation)
+	{
+		ShallowWaterNiagaraSimulation->SetVariableInt(FName("NumImpacts"), 0);
+		CurrentNumImpacts = 0;
+	}
+}
+
+void UShallowWaterSubsystem::RegisterImpact(FVector ImpactPosition, FVector ImpactVelocity, float ImpactRadius)
+{
+	if (ShallowWaterNiagaraSimulation == nullptr)
+	{
+		return;
+	}
+	
+	FHitResult OutHit;
+	GetWorld()->LineTraceSingleByChannel(OutHit, ImpactPosition + FVector(0,0,10.f), ImpactPosition + FVector(0,0,-10.f),
+		GetImpactCollisionChannel());
+	AWaterBody* WaterBody = Cast<AWaterBody>(OutHit.GetActor());
+	if (WaterBody == nullptr)
+	{
+		return;
+	}
+	Tracker_Directs.Add(FShallowWaterCollisionTracker_Direct(GetWorld()->GetTimeSeconds(), GSWImpactTrackerActiveForSeconds, WaterBody));
+	
+	if (!ShallowWaterNiagaraSimulation->IsActive())
+	{
+		// queue up impacts we missed
+		PendingImpact TmpImpact;
+		TmpImpact.ImpactPosition = ImpactPosition;
+		TmpImpact.ImpactVelocity = ImpactVelocity;
+		TmpImpact.ImpactRadius = ImpactRadius;
+		PendingImpacts.Add(TmpImpact);
+
+		return;
+	}
+
+	SetNiagaraImpactVariables(ImpactPosition, ImpactVelocity, ImpactRadius);
+}
+
+// these impacts are a frame or two late, and we only want to update the Niagara system and not 
+// perform an additional overlap test
+void UShallowWaterSubsystem::FlushPendingImpacts()
+{
+	for (PendingImpact CurrPendingImpact : PendingImpacts)
+	{
+		SetNiagaraImpactVariables(CurrPendingImpact.ImpactPosition, CurrPendingImpact.ImpactVelocity, CurrPendingImpact.ImpactRadius);
+	}
+	PendingImpacts.Reset();
+}
+
+void UShallowWaterSubsystem::SetNiagaraImpactVariables(FVector ImpactPosition, FVector ImpactVelocity, float ImpactRadius)
+{
+	UNiagaraDataInterfaceArrayFunctionLibrary::SetNiagaraArrayPositionValue(ShallowWaterNiagaraSimulation, "User.ImpactPositionArray", CurrentNumImpacts, ImpactPosition, false);
+	UNiagaraDataInterfaceArrayFunctionLibrary::SetNiagaraArrayVectorValue(ShallowWaterNiagaraSimulation, "User.ImpactVelocityArray", CurrentNumImpacts, ImpactVelocity, false);
+	UNiagaraDataInterfaceArrayFunctionLibrary::SetNiagaraArrayFloatValue(ShallowWaterNiagaraSimulation, "User.ImpactRadiusArray", CurrentNumImpacts, ImpactRadius, false);
+
+	++CurrentNumImpacts;
+	ShallowWaterNiagaraSimulation->SetVariableInt(FName("NumImpacts"), CurrentNumImpacts);
+}
+
+void UShallowWaterSubsystem::SetWaterBodyMIDParameters(AWaterBody* WaterBody)
+{
+	if (WaterBody)
+	{
+		if (IsShallowWaterInitialized())
+		{
+			TryUpdateWaterBodyMIDParameters(WaterBody->GetWaterBodyComponent());
+		}
+		else
+		{
+			// On BeginPlay this might not be ready yet
+			PendingWaterBodiesToSetMIDOnInitialize.AddUnique(WaterBody);
+		}
+	}
+}
+
+void UShallowWaterSubsystem::TryUpdateWaterBodyMIDParameters(UWaterBodyComponent* WaterBodyComponent)
+{
+	if (WaterBodyComponent == nullptr || WaterBodyComponentsWithProperMIDParameters.Contains(WaterBodyComponent))
+	{
+		// UE_LOG(LogShallowWater, Warning, TEXT("TryUpdateWaterBodyMIDParameters failed"));
+		return;
+	}
+	WaterBodyComponentsWithProperMIDParameters.Add(WaterBodyComponent);
+	
+	if (UMaterialInstanceDynamic* WaterMID = WaterBodyComponent->GetWaterMaterialInstance(); WaterMID)
+	{
+		WaterMID->SetTextureParameterValue(Settings->NormalRTMaterialName, NormalRT);
+		WaterMID->SetScalarParameterValue(FName(TEXT("DEV_UseNewShallowWaterSubsystem")), 1.f);   // Temp before replacing OG system
+	}
+	if (UMaterialInstanceDynamic* WaterMID = WaterBodyComponent->GetRiverToOceanTransitionMaterialInstance(); WaterMID)
+	{
+		WaterMID->SetTextureParameterValue(Settings->NormalRTMaterialName, NormalRT);
+		WaterMID->SetScalarParameterValue(FName(TEXT("DEV_UseNewShallowWaterSubsystem")), 1.f);   // Temp before replacing OG system
+	}
+	if (UMaterialInstanceDynamic* WaterMID = WaterBodyComponent->GetRiverToLakeTransitionMaterialInstance(); WaterMID)
+	{
+		WaterMID->SetTextureParameterValue(Settings->NormalRTMaterialName, NormalRT);
+		WaterMID->SetScalarParameterValue(FName(TEXT("DEV_UseNewShallowWaterSubsystem")), 1.f);   // Temp before replacing OG system
+	}
+}
+
+void UShallowWaterSubsystem::RegisterPhysicsAssetProxiesDataAsset(UShallowWaterPhysicsAssetOverridesDataAsset* Proxies)
+{
+	if (Proxies == nullptr)
+	{
+		ensure(false);
+		return;
+	}
+	for (const TTuple<FGameplayTag, FShallowWaterPhysicsAssetOverride>& Override : Proxies->Overrides)
+	{
+		if (RegisteredPhysicsAssetProxies.Contains(Override.Key))
+		{
+			UE_LOG(LogShallowWater, Log, TEXT("Physics Asset Override in %s is overwriting an existing Override. GameplayTag = %s. This could be intended."),
+				*Proxies->GetName(), *Override.Key.ToString())
+		}
+	}
+	RegisteredPhysicsAssetProxies.Append(Proxies->Overrides);
+}
+
+TSet<AWaterBody*> UShallowWaterSubsystem::GetAllOverlappingWaterBodiesAndUpdateCollisionTrackers() 
+{
+	TSet<AWaterBody*> Result = GetOverlappingWaterBodiesFromPawns();
+	Result.Append(GetOverlappingWaterBodiesFromActorTrackersAndUpdate());
+	Result.Append(GetOverlappingWaterBodiesFromDirectTrackersAndUpdate());
+	return Result;
+}
+
+void UShallowWaterSubsystem::AddCollisionTrackerForActor(AActor* CollisionTrackerActor, float MaxLifespan)
+{
+	if(CollisionTrackerActor == nullptr)
+	{
+		return;
+	}
+		
+	bool bFound = false;
+	for (FShallowWaterCollisionTracker_Actor& Tracker : Tracker_Actors)
+	{
+		if (Tracker.CollisionActor == CollisionTrackerActor)
+		{
+			Tracker.TimeSpawned = GetWorld()->GetTimeSeconds();
+			bFound = true;
+			break;
+		}
+	}
+	if (!bFound)
+	{
+		Tracker_Actors.Add(FShallowWaterCollisionTracker_Actor(GetWorld()->GetTimeSeconds(), MaxLifespan, CollisionTrackerActor));
+	}
+}
+
+void UShallowWaterSubsystem::RemoveCollisionTrackerForActor(AActor* CollisionTrackerActor)
+{
+	if(CollisionTrackerActor == nullptr)
+	{
+		return;
+	}
+
+	for (int32 Index = Tracker_Actors.Num()-1; Index >= 0; Index--)
+	{
+		if (Tracker_Actors[Index].CollisionActor == CollisionTrackerActor)
+		{
+			Tracker_Actors.RemoveAt(Index);
+		}
+	}
+}
+
+TSet<AWaterBody*> UShallowWaterSubsystem::GetOverlappingWaterBodiesFromPawns() const
+{
+	TSet<AWaterBody*> Result;
+	TArray<APawn*> Pawns = GetPawnsInRange();
+	for (const APawn* Pawn : Pawns)
+	{
+		TArray<AActor*> OverlappingActors;
+		Pawn->GetOverlappingActors(OverlappingActors, AWaterBody::StaticClass());
+		for (AActor* WaterBody : OverlappingActors)
+		{
+			if (AWaterBody* Actor = Cast<AWaterBody>(WaterBody))
+			{
+				Result.Add(Actor);
+			}
+		}
+	}
+	return Result;
+}
+
+TSet<AWaterBody*> UShallowWaterSubsystem::GetOverlappingWaterBodiesFromActorTrackersAndUpdate()
+{
+	TSet<AWaterBody*> Result;
+	for (int32 Index = Tracker_Actors.Num() - 1; Index >= 0; Index--)
+	{
+		if (!Tracker_Actors[Index].IsValid(GetWorld()->GetTimeSeconds()))
+		{
+			Tracker_Actors.RemoveAt(Index);
+			continue;
+		}
+		Result.Append(Tracker_Actors[Index].GetOverlappingWaterBodies());
+	}
+	return Result;
+}
+
+TSet<AWaterBody*> UShallowWaterSubsystem::GetOverlappingWaterBodiesFromDirectTrackersAndUpdate()
+{
+	TSet<AWaterBody*> Result;
+	for (int32 Index = Tracker_Directs.Num() - 1; Index >= 0; Index--)
+	{
+		if (!Tracker_Directs[Index].IsValid(GetWorld()->GetTimeSeconds()))
+		{
+			Tracker_Directs.RemoveAt(Index);
+			continue;
+		}
+		if(AWaterBody* WaterBody = Tracker_Directs[Index].GetOverlappingWaterBody())
+		{
+			Result.Add(WaterBody);
+		}
+	}
+	return Result;
+}
+
+void UShallowWaterSubsystem::UpdateOverlappingWaterBodiesHistory(TArray<AWaterBody*> OverlappingWaterBodies)
+{
+	if (!OverlappingWaterBodies.IsEmpty())
+	{
+		LastTimeOverlappingAnyWaterBody = GetWorld()->GetTimeSeconds();
+		LastOverlappingWaterBodies_Internal.Reset();
+		for (AWaterBody* WaterBody : OverlappingWaterBodies)
+		{
+			LastOverlappingWaterBodies_Internal.Add(WaterBody);  			
+		}
+	}
+}
+
+bool UShallowWaterSubsystem::ShouldSimulateThisFrame() const
+{
+	return (GetWorld()->GetTimeSeconds() - LastTimeOverlappingAnyWaterBody <= GSWRemainActiveForSeconds);
+}
+
+int32 UShallowWaterSubsystem::UpdateActivePawns()
+{
+	const APawn* CursorPawn = GetTheMostRelevantPlayerPawn();
+	if (CursorPawn == nullptr)
+	{
+		ActivePawns.Empty();
+		return 0;
+	}
+	FVector CursorPawnLocation = CursorPawn->GetActorLocation();
+	
+	// Remove irrelevant pawns from ActivePawns
+	ActivePawns.RemoveAll([&](TWeakObjectPtr<APawn> Pawn)
+	{
+		if(Pawn == nullptr)
+		{
+			return true;
+		}
+		if ((Pawn->GetActorLocation() - CursorPawnLocation).SquaredLength() > GetColliderMaxRange() * GetColliderMaxRange())
+		{
+			return true;
+		}
+		return false;
+	});
+
+	const int32 SlotsLeft = Settings->MaxActivePawnNum - ActivePawns.Num();
+	int32 NewPawnsAdded = 0;
+	if (SlotsLeft > 0)
+	{
+		// Get pawns nearby who become relevant this frame
+		TArray<APawn*> RelevantPawns = GetPawnsInRange();
+		RelevantPawns.RemoveAll([&](APawn* Pawn)
+		{
+			for (const TWeakObjectPtr<APawn> Other : ActivePawns)
+			{
+				if (Pawn == nullptr || Pawn == Other.Get())
+				{
+					return true;
+				}
+			}
+			return false;
+		});
+
+		// Sort pawn candidates
+		RelevantPawns.Sort([CursorPawnLocation] (const APawn& Left, const APawn& Right)
+		{
+			// Compare 3D distance since the other player right above you, although might be pretty far away, are going to drop on your face fast
+			// Thus prioritized over the player in water but a little further away
+			return (Left.GetActorLocation() - CursorPawnLocation).SquaredLength()
+			< (Right.GetActorLocation() - CursorPawnLocation).SquaredLength();
+		});
+
+		NewPawnsAdded = FMath::Min(SlotsLeft, RelevantPawns.Num());
+		for (int32 Index = 0; Index < NewPawnsAdded; Index++)
+		{
+			ActivePawns.Add(RelevantPawns[Index]);
+		}
+
+#if ENABLE_DRAW_DEBUG
+		if (Settings->bVisualizeActivePawn)
+		{
+			// Use Niagara to draw debug should be better. If yes delete this
+		}
+#endif
+	}
+	
+	return NewPawnsAdded;
+}
+
+void UShallowWaterSubsystem::GatherContextsFromPawns(TArray<TWeakObjectPtr<APawn>> InActivePawns)
+{
+	for (const TWeakObjectPtr<APawn>& Pawn : InActivePawns)
+	{
+		if (APawn* PawnPtr = Pawn.Get())
+		{
+			TOptional<FShallowWaterCollisionContext> Context = GetCollisionContextFromPawn(PawnPtr);
+			if (Context.IsSet())
+			{
+				// AddUnique() since multiple pawns can return the same Context. E.g. a multi-seat vehicle
+				PendingContexts.AddUnique(Context.GetValue());
+			}
+		}
+	}
+}
+
+TOptional<FShallowWaterCollisionContext> UShallowWaterSubsystem::GetCollisionContextFromPawn(APawn* InPawn) const
+{
+	TOptional<FShallowWaterCollisionContext> Result;
+	USkeletalMeshComponent* Component = nullptr;
+	const ACharacter* Character = Cast<ACharacter>(InPawn);
+	if (Character)
+	{
+		Component = Character->GetMesh();
+	}
+	else
+	{
+		TArray<USkeletalMeshComponent*> SKMs;
+		InPawn->GetComponents<USkeletalMeshComponent>(SKMs);
+		if (!SKMs.IsEmpty())
+		{
+			Component = SKMs[0];
+		}
+	}
+	if (Component)
+	{
+		Result = FShallowWaterCollisionContext(EShallowWaterCollisionContextType::Pawn, Component); 
+	}
+	return Result;
+}
+
+void UShallowWaterSubsystem::CleanUpVehicleCollisionProxies()
+{
+	TArray<TTuple<const FShallowWaterCollisionContext, USkeletalMeshComponent*>> EntriesToRemove;
+	for (TTuple<const FShallowWaterCollisionContext, USkeletalMeshComponent*> Entry : VehicleCollisionProxies)
+	{
+		// Shotgun approach to cover all weird possibilities
+		// E.g. if the Vehicle component is destroyed but not the owning actor somehow
+		if ( Entry.Key.Component == nullptr || Entry.Key.Component->IsBeingDestroyed() ||
+			Entry.Value == nullptr || Entry.Value->IsBeingDestroyed())
+		{
+			EntriesToRemove.Add(Entry);
+		}
+	}
+	for (TTuple<const FShallowWaterCollisionContext, USkeletalMeshComponent*> Entry : EntriesToRemove)
+	{
+		DisableCollisionForVehicle(Entry.Key);
+	}
+}
+
+void UShallowWaterSubsystem::UpdateCollisionForPendingContexts()
+{
+	CleanUpVehicleCollisionProxies();
+	
+	// Find components that need their collision disabled
+	TArray<FShallowWaterCollisionContext> ContextsToDisable = PreviousContexts;
+	ContextsToDisable.RemoveAll([&] (FShallowWaterCollisionContext Context)
+	{
+		return PendingContexts.Find(Context) != INDEX_NONE;
+	});
+
+	// Find components that don't already exist
+	TArray<FShallowWaterCollisionContext> ContextsToEnable = PendingContexts;
+	ContextsToEnable.RemoveAll([&] (FShallowWaterCollisionContext Context)
+	{
+		return PreviousContexts.Find(Context) != INDEX_NONE;
+	});
+	
+
+	
+	for (const FShallowWaterCollisionContext& Context : ContextsToDisable)
+	{
+		DisableCollisionForContext(Context);
+	}
+	for (const FShallowWaterCollisionContext& Context : ContextsToEnable)
+	{
+		EnableCollisionForContext(Context);
+	}
+
+	PreviousContexts = PendingContexts;
+	PendingContexts.Empty();
+}
+
+void UShallowWaterSubsystem::EnableCollisionForContext(const FShallowWaterCollisionContext& Context)
+{
+	if (Context.Component == nullptr)
+	{
+		ensure(false);
+		UE_LOG(LogShallowWater, Warning, TEXT("EnableCollisionForContext() - Context Component is nullptr"));
+		return;
+	}
+	switch (Context.Type)
+	{
+	case EShallowWaterCollisionContextType::Pawn:
+		
+		Context.Component->ComponentTags.AddUnique(ColliderComponentTag);
+		break;
+		
+	case EShallowWaterCollisionContextType::Vehicle:
+		{
+		// Spawn proxy SKM, attach to vehicle
+		const FString BaseName = TEXT("FluidsimCollisionProxy");
+		const FName CompName = MakeUniqueObjectName(this, USkeletalMeshComponent::StaticClass(), FName(*BaseName));
+		USkeletalMeshComponent* ProxyComp = NewObject<USkeletalMeshComponent>(Context.Component->GetOwner(), USkeletalMeshComponent::StaticClass(), CompName);
+
+		// Spawn an empty dummy SKM component to be used as Collision 
+		// Apply PhysicsAsset override if defined in the data asset
+		// #todo confirm if SetLeaderPoseComponent actually works if Mesh is set to empty #PLAY-29387
+		UPhysicsAsset* PhysicsAssetOverride = nullptr;
+		if (!RegisteredPhysicsAssetProxies.IsEmpty())
+		{
+			const FGameplayTagContainer VehicleTags = GetVehicleTags(Context);
+
+			const FShallowWaterPhysicsAssetOverride* FoundOverride = nullptr;
+			for (FGameplayTag Tag : VehicleTags)
+			{
+				FoundOverride = RegisteredPhysicsAssetProxies.Find(Tag);
+				if (FoundOverride)
+				{
+					break;
+				}
+			}
+			if (FoundOverride)
+			{
+				TSoftObjectPtr<UPhysicsAsset> TmpPhysicsAsset = FoundOverride->PhysicsAsset;
+
+				UAssetManager::GetStreamableManager().RequestAsyncLoad(TmpPhysicsAsset.ToSoftObjectPath(),
+					FStreamableDelegate::CreateWeakLambda(this, [this, TmpPhysicsAsset, ProxyComp, Context]()
+						{								
+							if (TmpPhysicsAsset.IsValid())
+							{
+								UPhysicsAsset* PhysicsAssetOverride = TmpPhysicsAsset.Get();
+
+								ProxyComp->SetPhysicsAsset(PhysicsAssetOverride);
+
+								ProxyComp->ComponentTags.AddUnique(ColliderComponentTag);
+
+								ProxyComp->SetupAttachment(Context.Component->GetOwner()->GetRootComponent());
+								ProxyComp->SetWorldTransform(Context.Component->GetComponentTransform());
+								ProxyComp->SetLeaderPoseComponent(Context.Component);
+								ProxyComp->SetVisibility(false);
+								ProxyComp->SetCollisionEnabled(ECollisionEnabled::NoCollision);
+								ProxyComp->RegisterComponent();
+
+								VehicleCollisionProxies.Add(Context, ProxyComp);
+							}
+						})
+				);
+			}
+		}
+		break;
+		}
+		
+	default:
+		;
+	}
+}
+
+void UShallowWaterSubsystem::DisableCollisionForContext(const FShallowWaterCollisionContext& Context)
+{
+	if (!Context.IsValidAndAlive())
+	{
+		switch (Context.Type)
+		{
+		case EShallowWaterCollisionContextType::Pawn:
+			break;
+		case EShallowWaterCollisionContextType::Vehicle:
+			DisableCollisionForVehicle(Context);
+			break;
+		default:
+			;
+		}
+		// Otherwise no care needed because the component is dead
+		return;
+	}
+
+	switch (Context.Type)
+	{
+	case EShallowWaterCollisionContextType::Pawn:
+		Context.Component->ComponentTags.Remove(ColliderComponentTag);
+		break;
+	case EShallowWaterCollisionContextType::Vehicle:
+		DisableCollisionForVehicle(Context);
+		break;
+	default:
+		;
+	}
+}
+
+void UShallowWaterSubsystem::DisableCollisionForVehicle(const FShallowWaterCollisionContext& Context)
+{
+	if (Context.Type != EShallowWaterCollisionContextType::Vehicle)
+	{
+		return;
+	}
+
+	USkeletalMeshComponent** ProxyPointer = VehicleCollisionProxies.Find(Context);
+	if (ProxyPointer == nullptr)
+	{
+		ensure(false);
+		return;
+	}
+	
+	USkeletalMeshComponent* Proxy = *ProxyPointer;
+	if (Proxy && !Proxy->IsBeingDestroyed())
+	{
+		Proxy->DestroyComponent();
+	}
+	VehicleCollisionProxies.Remove(Context);
+}
+
+float UShallowWaterSubsystem::GetColliderMaxRange() const
+{
+	return GetGridSize() * 0.5f;
+}
+
+TArray<APawn*> UShallowWaterSubsystem::GetPawnsInRange(const bool bShouldSortBySignificance) const
+{
+	const APlayerController* Controller = WeakPlayerController.Get();
+	APawn* CursorPawn = nullptr;
+	if (Controller)
+	{
+		CursorPawn = GetNonSpectatorPawnFromWeakController();
+	}
+	
+	FVector ObservingLocation;
+	if (CursorPawn)
+	{
+		ObservingLocation = CursorPawn->GetActorLocation();
+	}
+	else
+	{
+		TOptional<FVector> CamLoc = GetCameraLocationFromWeakController();
+		if (CamLoc.IsSet())
+		{
+			ObservingLocation = CamLoc.GetValue();
+		}
+		else
+		{
+			// Not even camera location is available. Game state not valid.
+			return TArray<APawn*>();
+		}
+	}
+
+	return GetPawnsInRange(ObservingLocation, bShouldSortBySignificance);
+}
+
+TArray<APawn*> UShallowWaterSubsystem::GetPawnsInRange(const FVector ObservingLocation, const bool bShouldSortBySignificance) const
+{
+	TArray<AActor*> FoundActors;
+	UGameplayStatics::GetAllActorsOfClass(Cast<UObject>(GetWorld()), APawn::StaticClass(), FoundActors);
+	TArray<APawn*> Results;
+	for (AActor* Actor : FoundActors)
+	{
+		const float Distance = FVector::Distance(ObservingLocation, Actor->GetActorLocation());
+		if (Distance <= GetColliderMaxRange())
+		{
+			Results.Add(Cast<APawn>(Actor));
+		}
+	}
+
+	if (bShouldSortBySignificance)
+	{
+		Results.Sort([ObservingLocation](APawn& PawnA, APawn& PawnB)
+		{
+			return FVector::DistSquared(ObservingLocation, PawnA.GetActorLocation()) < FVector::DistSquared(ObservingLocation, PawnB.GetActorLocation());
+		});
+	}
+	
+	return Results;
+}
+
+void UShallowWaterSubsystem::TryGetOrWaitForWaterInfoTextureFromWaterBodies(TSet<AWaterBody*> CurrentWaterBodies)
+{
+    for (const AWaterBody* CurrentWaterBody : CurrentWaterBodies)
+    {
+	    if (const UWaterBodyComponent* WaterBodyComp = CurrentWaterBody->GetWaterBodyComponent())
+	    {
+	    	if (AWaterZone* WaterZone = WaterBodyComp->GetWaterZone())
+	    	{
+	    		const TObjectPtr<UTextureRenderTarget2D> NewWaterInfoTexture = WaterZone->WaterInfoTexture;
+
+	    		if (NewWaterInfoTexture == nullptr)
+	    		{
+	    			WaterZone->GetOnWaterInfoTextureCreated().RemoveDynamic(this, &UShallowWaterSubsystem::OnWaterInfoTextureCreated);
+	    			WaterZone->GetOnWaterInfoTextureCreated().AddDynamic(this, &UShallowWaterSubsystem::OnWaterInfoTextureCreated);
+	    		}
+	    		else
+	    		{
+	    			OnWaterInfoTextureCreated(NewWaterInfoTexture);
+	    		}
+
+	    		// Currently there can only be one unique WaterInfoTexture
+	    		break;
+	    	}
+	    }
+    }
+}
+
+void UShallowWaterSubsystem::OnWaterInfoTextureCreated(const UTextureRenderTarget2D* InWaterInfoTexture)
+{
+	if(InWaterInfoTexture == nullptr)
+	{
+		ensureMsgf(false, TEXT("UShallowWaterSubsystem::OnWaterInfoTextureCreated was called with NULL WaterInfoTexture"));
+		return;
+	}
+	WaterInfoTexture = InWaterInfoTexture;
+	if (ShallowWaterNiagaraSimulation)
+	{
+		ShallowWaterNiagaraSimulation->SetVariableTexture(FName("WaterInfoTexture"), Cast<UTexture>(const_cast<UTextureRenderTarget2D*>(WaterInfoTexture)));
+	}
+}
+
+void UShallowWaterSubsystem::OnLocalPlayerControllerBecomesValid(APlayerController* InPlayerController)
+{
+	if (InPlayerController == nullptr)
+	{
+		UE_LOG(LogShallowWater, Log, TEXT("OnLocalPlayerControllerBecomesValid() returned nullptr"))
+		return;
+	}
+	
+	UE_LOG(LogShallowWater, Log, TEXT("OnLocalPlayerControllerBecomesValid() returned PC: %s"), *InPlayerController->GetFullName())
+	WeakPlayerController = InPlayerController;
+	if (APawn* const Pawn = InPlayerController->GetPawn())
+	{
+		OnLocalPlayerPawnBecomesValid(nullptr, Pawn);
+	}
+	InPlayerController->OnPossessedPawnChanged.RemoveDynamic(this, &ThisClass::OnLocalPlayerPawnBecomesValid);
+	InPlayerController->OnPossessedPawnChanged.AddDynamic(this, &ThisClass::OnLocalPlayerPawnBecomesValid);
+}
+
+void UShallowWaterSubsystem::OnLocalPlayerPawnBecomesValid(APawn* OldPawn, APawn* NewPawn)
+{
+	if (NewPawn == nullptr)
+	{
+		UE_LOG(LogShallowWater, Log, TEXT("OnLocalPlayerPawnBecomesValid() returned nullptr"))
+		return;
+	}
+	
+	UE_LOG(LogShallowWater, Log, TEXT("OnLocalPlayerPawnBecomesValid() returned Pawn: %s"), *NewPawn->GetFullName())
+	// #todo Should re-initialize if pawn changed, or uninitialize if pawn lost 
+	if (!IsShallowWaterInitialized())
+	{
+		InitializeShallowWater();
+	}
+	else
+	{
+		UE_LOG(LogShallowWater, Log, TEXT("OnLocalPlayerPawnBecomesValid called but subsystem is already initialized."));
+	}
+}
+
+//UE_ENABLE_OPTIMIZATION_SHIP
diff --git a/Engine/Plugins/Experimental/WaterAdvanced/Source/WaterAdvanced/Private/WaterAdvanced.cpp b/Engine/Plugins/Experimental/WaterAdvanced/Source/WaterAdvanced/Private/WaterAdvanced.cpp
new file mode 100644
--- /dev/null
+++ b/Engine/Plugins/Experimental/WaterAdvanced/Source/WaterAdvanced/Private/WaterAdvanced.cpp
@@ -0,0 +1,20 @@
+// Copyright Epic Games, Inc. All Rights Reserved.
+
+#include "WaterAdvanced.h"
+
+#define LOCTEXT_NAMESPACE "FWaterAdvancedModule"
+
+void FWaterAdvancedModule::StartupModule()
+{
+	// This code will execute after your module is loaded into memory; the exact timing is specified in the .uplugin file per-module
+}
+
+void FWaterAdvancedModule::ShutdownModule()
+{
+	// This function may be called during shutdown to clean up your module.  For modules that support dynamic reloading,
+	// we call this function before unloading the module.
+}
+
+#undef LOCTEXT_NAMESPACE
+	
+IMPLEMENT_MODULE(FWaterAdvancedModule, WaterAdvanced)
\ No newline at end of file
diff --git a/Engine/Plugins/Experimental/WaterAdvanced/Source/WaterAdvanced/Public/ShallowWaterCommon.h b/Engine/Plugins/Experimental/WaterAdvanced/Source/WaterAdvanced/Public/ShallowWaterCommon.h
new file mode 100644
--- /dev/null
+++ b/Engine/Plugins/Experimental/WaterAdvanced/Source/WaterAdvanced/Public/ShallowWaterCommon.h
@@ -0,0 +1,32 @@
+﻿// Copyright Epic Games, Inc. All Rights Reserved.
+
+#pragma once
+
+#include "CoreMinimal.h"
+#include "ShallowWaterCommon.generated.h"
+
+DECLARE_LOG_CATEGORY_EXTERN(LogShallowWater, Log, All);
+
+class UPhysicsAsset;
+
+USTRUCT(BlueprintType)
+struct WATERADVANCED_API FShallowWaterSimParameters 
+{
+	GENERATED_BODY()
+
+	UPROPERTY(EditAnywhere, Category = "Shallow Water")
+	float WorldGridSize = 2000.f;
+	UPROPERTY(EditAnywhere, Category = "Shallow Water")
+	int32 ResolutionMaxAxis = 512;
+	UPROPERTY(EditAnywhere, Category = "Shallow Water")
+	int32 MaxImpactsPerFrame = 50;
+};
+
+USTRUCT(BlueprintType)
+struct WATERADVANCED_API FShallowWaterPhysicsAssetOverride
+{
+	GENERATED_BODY();
+
+	UPROPERTY(EditAnywhere, Category = "Shallow Water")
+	TSoftObjectPtr<UPhysicsAsset> PhysicsAsset;
+};
\ No newline at end of file
diff --git a/Engine/Plugins/Experimental/WaterAdvanced/Source/WaterAdvanced/Public/ShallowWaterPhysicsAssetOverridesDataAsset.h b/Engine/Plugins/Experimental/WaterAdvanced/Source/WaterAdvanced/Public/ShallowWaterPhysicsAssetOverridesDataAsset.h
new file mode 100644
--- /dev/null
+++ b/Engine/Plugins/Experimental/WaterAdvanced/Source/WaterAdvanced/Public/ShallowWaterPhysicsAssetOverridesDataAsset.h
@@ -0,0 +1,24 @@
+﻿// Copyright Epic Games, Inc. All Rights Reserved.
+
+#pragma once
+
+#include "CoreMinimal.h"
+#include "Engine/DataAsset.h"
+#include "ShallowWaterCommon.h"
+#include "GameplayTagContainer.h"
+#include "ShallowWaterPhysicsAssetOverridesDataAsset.generated.h"
+
+class USkeletalMesh;
+
+/**
+ * 
+ */
+UCLASS(Blueprintable)
+class WATERADVANCED_API UShallowWaterPhysicsAssetOverridesDataAsset : public UDataAsset
+{
+	GENERATED_BODY()
+
+public:
+	UPROPERTY(EditAnywhere, Category = "Shallow Water")
+	TMap<FGameplayTag, FShallowWaterPhysicsAssetOverride> Overrides;
+};
diff --git a/Engine/Plugins/Experimental/WaterAdvanced/Source/WaterAdvanced/Public/ShallowWaterSettings.h b/Engine/Plugins/Experimental/WaterAdvanced/Source/WaterAdvanced/Public/ShallowWaterSettings.h
new file mode 100644
--- /dev/null
+++ b/Engine/Plugins/Experimental/WaterAdvanced/Source/WaterAdvanced/Public/ShallowWaterSettings.h
@@ -0,0 +1,47 @@
+﻿// Copyright Epic Games, Inc. All Rights Reserved.
+
+#pragma once
+
+#include "CoreMinimal.h"
+#include "Engine/DeveloperSettings.h"
+#include "NiagaraSystem.h"
+#include "ShallowWaterCommon.h"
+#include "ShallowWaterPhysicsAssetOverridesDataAsset.h"
+#include "Materials/MaterialParameterCollection.h"
+#include "ShallowWaterSettings.generated.h"
+
+UCLASS(config = Engine, defaultconfig, meta=(DisplayName="Water Advanced"))
+class WATERADVANCED_API UShallowWaterSettings : public UDeveloperSettings
+{
+	GENERATED_BODY()
+	UShallowWaterSettings();
+	
+public:
+	virtual FName GetCategoryName() const override;
+	
+	UPROPERTY(Config, BlueprintReadOnly, EditAnywhere, Category="Shallow Water Simulation")
+	TSoftObjectPtr<UNiagaraSystem> DefaultShallowWaterNiagaraSimulation;
+	UPROPERTY(Config, BlueprintReadOnly, EditAnywhere, Category="Shallow Water Simulation")
+	FShallowWaterSimParameters ShallowWaterSimParameters;
+	UPROPERTY(Config, BlueprintReadOnly, EditAnywhere, Category="Shallow Water Simulation")
+	TSoftObjectPtr<UMaterialParameterCollection> WaterMPC;
+	UPROPERTY(Config, BlueprintReadOnly, EditAnywhere, Category="Shallow Water Simulation")
+	FName GridCenterMPCName = FName("SimLocation");
+	UPROPERTY(Config, BlueprintReadOnly, EditAnywhere, Category="Shallow Water Simulation")
+	FName WorldGridSizeMPCName = FName("FluidSimSize");
+	UPROPERTY(Config, BlueprintReadOnly, EditAnywhere, Category="Shallow Water Simulation")
+	FName ResolutionMaxAxisMPCName = FName("FluidSimResolution");
+	UPROPERTY(Config, BlueprintReadOnly, EditAnywhere, Category="Shallow Water Simulation")
+	FName NormalRTMaterialName = FName("NormalAndHeight");
+
+	UPROPERTY(Config, BlueprintReadOnly, EditAnywhere, Category="Interaction")
+	int32 MaxActivePawnNum = 6;
+	UPROPERTY(Config, BlueprintReadOnly, EditAnywhere, Category="Interaction")
+	int32 MaxImpulseForceNum = 6;
+	// Default overrides. Game Feature Plugin can register their own overrides asset with subsystem RegisterPhysicsAssetOverridesDataAsset()
+	UPROPERTY(Config, BlueprintReadOnly, EditAnywhere, Category="Interaction")
+	TSoftObjectPtr<UShallowWaterPhysicsAssetOverridesDataAsset> PhysicsAssetProxiesDataAsset;
+
+	UPROPERTY(Config, BlueprintReadOnly, EditAnywhere, Category="Debug")
+	bool bVisualizeActivePawn;
+};
diff --git a/Engine/Plugins/Experimental/WaterAdvanced/Source/WaterAdvanced/Public/ShallowWaterSubsystem.h b/Engine/Plugins/Experimental/WaterAdvanced/Source/WaterAdvanced/Public/ShallowWaterSubsystem.h
new file mode 100644
--- /dev/null
+++ b/Engine/Plugins/Experimental/WaterAdvanced/Source/WaterAdvanced/Public/ShallowWaterSubsystem.h
@@ -0,0 +1,315 @@
+﻿// Copyright Epic Games, Inc. All Rights Reserved.
+
+#pragma once
+
+#include "CoreMinimal.h"
+#include "GameplayTagContainer.h"
+#include "ShallowWaterSettings.h"
+#include "Engine/TextureRenderTarget2D.h"
+#include "Components/SceneComponent.h"
+#include "Components/SkeletalMeshComponent.h"
+#include "Subsystems/WorldSubsystem.h"
+#include "GameplayTagContainer.h"
+
+#include "ShallowWaterSubsystem.generated.h"
+
+class AWaterBody;
+class UWaterBodyComponent;
+class UNiagaraComponent;
+class APawn;
+struct FGameplayTag;
+
+UENUM(BlueprintType)
+enum class EShallowWaterCollisionContextType : uint8
+{
+	Pawn,
+	Vehicle, // Pawn as driver or pawn as passenger
+	Custom
+};
+
+// Shallow Water Rigid Body collision
+USTRUCT(BlueprintType)
+struct WATERADVANCED_API FShallowWaterCollisionContext
+{
+	GENERATED_BODY()
+	
+	FShallowWaterCollisionContext() {};
+	FShallowWaterCollisionContext(EShallowWaterCollisionContextType InType, USkeletalMeshComponent* InComponent)
+		: Type(InType)
+		, Component(InComponent)
+	{
+		if(InComponent == nullptr)
+		{
+			ensureMsgf(false, TEXT("FShallowWaterCollisionContext constructor was given nullptr component as input"));
+			return;
+		}
+		UniqueID = Component->GetUniqueID();
+	}
+	EShallowWaterCollisionContextType Type = EShallowWaterCollisionContextType::Pawn;
+	UPROPERTY()
+	USkeletalMeshComponent* Component = nullptr;
+
+	// #todo Theoretically not enough as unique identification since Component->GetUniqueID() is "reused so it is only unique while the object is alive"
+	uint32 UniqueID = 0;
+
+	bool IsValidAndAlive() const { return Component != nullptr && !Component->IsBeingDestroyed(); }
+
+	bool operator==(const FShallowWaterCollisionContext& Other) const
+	{
+		return UniqueID == Other.UniqueID;
+	}
+	friend uint32 GetTypeHash(const FShallowWaterCollisionContext& Context)
+	{
+		return Context.UniqueID;
+	}
+};
+
+/*
+ * CollisionTracker that records the actor affecting the waterbody, used by logs or fishing lures etc.,
+ * where the collision is handled by Niagara Data Channel and the subsystem isn't aware of them without a tracking device. This is intended as a solution.
+ */
+USTRUCT(BlueprintType)
+struct WATERADVANCED_API FShallowWaterCollisionTracker_Actor
+{
+	GENERATED_BODY()
+
+	FShallowWaterCollisionTracker_Actor() {}
+	FShallowWaterCollisionTracker_Actor(float InTimeSpawned, float InLifespan, TWeakObjectPtr<AActor> InCollisionActor)
+		: TimeSpawned(InTimeSpawned), Lifespan(InLifespan), CollisionActor(InCollisionActor)
+	{}
+
+	TArray<AWaterBody*> GetOverlappingWaterBodies() const;
+	bool IsValid(const float CurrentTime) const { return CollisionActor.IsValid() && CurrentTime - TimeSpawned <= Lifespan; }
+
+	float TimeSpawned = 0;
+	float Lifespan = 10.f;
+	UPROPERTY()
+	TWeakObjectPtr<AActor> CollisionActor = nullptr;
+
+	bool operator==(const FShallowWaterCollisionTracker_Actor& Other) const
+	{
+		return CollisionActor == Other.CollisionActor;
+	}
+};
+
+/*
+ * CollisionTracker that directly records the water body affected, used by Impacts
+ */
+USTRUCT()
+struct WATERADVANCED_API FShallowWaterCollisionTracker_Direct
+{
+	GENERATED_BODY()
+
+	FShallowWaterCollisionTracker_Direct() {}
+	FShallowWaterCollisionTracker_Direct(float InTimeSpawned, float InLifespan, TWeakObjectPtr<AWaterBody> InWaterBody)
+		: TimeSpawned(InTimeSpawned), Lifespan(InLifespan), WaterBody(InWaterBody)
+	{}
+
+	AWaterBody* GetOverlappingWaterBody() const { return WaterBody.Get(); }
+	bool IsValid(const float CurrentTime) const { return WaterBody.IsValid() && CurrentTime - TimeSpawned <= Lifespan; }
+
+	float TimeSpawned = 0;
+	float Lifespan = 10.f;
+	UPROPERTY()
+	TWeakObjectPtr<AWaterBody> WaterBody = nullptr;
+};
+
+
+
+UCLASS(Abstract)
+class WATERADVANCED_API UShallowWaterSubsystem : public UTickableWorldSubsystem
+{
+	GENERATED_BODY()
+
+public:
+	UShallowWaterSubsystem();
+
+	virtual void PostInitialize() override;
+	virtual void OnWorldBeginPlay(UWorld& InWorld) override;
+	virtual bool ShouldCreateSubsystem(UObject* Outer) const override;
+	virtual void Tick(float DeltaTime) override;
+	virtual TStatId GetStatId() const override;
+
+	void InitializeShallowWater();
+	/*
+	 * Note: A world subsystem is always created and activated even when the plugin (GFP or not) it resides in is completely disabled
+	 * When the subclass of this class is a Game Feature Plugin. This can be used to limit when ShallowWater is actually enabled
+	 * by checking GameFeaturesSubsystem.GetPluginURLByName(PluginName, PluginURL))
+	 * and GameFeaturesSubsystem.IsGameFeaturePluginActive(PluginURL, true)
+	 */ 
+	virtual bool IsShallowWaterAllowedToInitialize() const { return true; }
+	bool IsShallowWaterInitialized() const;
+	APawn* GetNonSpectatorPawnFromWeakController() const;
+	TOptional<FVector> GetCameraLocationFromWeakController() const;
+	
+	// Think of this as a cursor that in most time locks on the current player pawn
+	// If the player doesn't have a physical pawn (e.g. spectating), the cursor jump around and lock on the nearest pawn to the camera
+	// If there is no relevant pawn (all eliminate), returns nullptr
+	APawn* GetTheMostRelevantPlayerPawn() const;
+	virtual FGameplayTagContainer GetVehicleTags(FShallowWaterCollisionContext Context) const PURE_VIRTUAL(UShallowWaterSubsystem::GetVehicleTags, return FGameplayTagContainer::EmptyContainer;)
+
+	void CreateRTs();
+	void InitializeParameters();
+	/*
+	 * @param WaterBody  
+	 */
+	void UpdateGridMovement();
+
+	// #note - these are nonblocking for now
+	UFUNCTION(BlueprintCallable, Category = "Shallow Water")
+	void ResetImpacts();
+	UFUNCTION(BlueprintCallable, Category = "Shallow Water")
+	void RegisterImpact(FVector ImpactPosition, FVector ImpactVelocity, float ImpactRadius);
+	
+	void FlushPendingImpacts();
+	void SetNiagaraImpactVariables(FVector ImpactPosition, FVector ImpactVelocity, float ImpactRadius);
+	// Override to return the ECC channel of bullets
+	virtual ECollisionChannel GetImpactCollisionChannel() PURE_VIRTUAL(UShallowWaterSubsystem::GetImpactCollisionChannel, return ECC_WorldDynamic;)
+	
+
+	// Manually set MID parameters for water bodies before they collide with any pawns or CollisionTrackers 
+	UFUNCTION(BlueprintCallable, Category = "Shallow Water")
+	void SetWaterBodyMIDParameters(AWaterBody* WaterBody);
+
+	void TryUpdateWaterBodyMIDParameters(UWaterBodyComponent* WaterBodyComponent);
+	float GetGridSize() const { return Settings->ShallowWaterSimParameters.WorldGridSize; }
+	int32 GetGridResolution() const { return Settings->ShallowWaterSimParameters.ResolutionMaxAxis; }
+	int32 GetMaxImpactsPerFrame() const { return Settings->ShallowWaterSimParameters.MaxImpactsPerFrame; }
+
+	/*
+	 * Add PA overrides. Designed to be called by Game Feature Plugins. 
+	 */
+	void RegisterPhysicsAssetProxiesDataAsset(UShallowWaterPhysicsAssetOverridesDataAsset* Proxies);
+
+protected:
+	// Asset can be set in Project Settings - Plugins - Water ShallowWaterSimulation
+	UPROPERTY(BlueprintReadOnly, VisibleAnywhere, Category="Shallow Water")
+	TObjectPtr<UNiagaraComponent> ShallowWaterNiagaraSimulation;
+
+	/*
+	 * WaterBody is used for:
+	 *     Get water data texture
+	 *     Get water zone
+	 *     Get character location projected onto water surface
+	 *     Check if character or vehicle is 'in water'
+	 * @return Should be the the water body touched by the most relevant pawn. If that's not available we try to find water body touched by nearby pawns, sorted by significance.
+	 */
+	UFUNCTION(BlueprintCallable, Category="Shallow Water")
+	virtual TSet<AWaterBody*> GetAllOverlappingWaterBodiesAndUpdateCollisionTrackers();
+	/*
+	 * @param MaxLifespan    Max lifespan in case RemoveCollisionTrackerForActor() is not called on tracker destroy, which could causing the sim to keep active indefinitely 
+	 */
+	UFUNCTION(BlueprintCallable, Category="Shallow Water")
+	void AddCollisionTrackerForActor(AActor* CollisionTrackerActor, float MaxLifespan = 60.f);
+	UFUNCTION(BlueprintCallable, Category="Shallow Water")
+	void RemoveCollisionTrackerForActor(AActor* CollisionTrackerActor);
+	
+	virtual TSet<AWaterBody*> GetOverlappingWaterBodiesFromPawns() const;
+	TSet<AWaterBody*> GetOverlappingWaterBodiesFromActorTrackersAndUpdate();
+	TSet<AWaterBody*> GetOverlappingWaterBodiesFromDirectTrackersAndUpdate();
+	void UpdateOverlappingWaterBodiesHistory(TArray<AWaterBody*> OverlappingWaterBodies);
+	bool ShouldSimulateThisFrame() const;
+
+	/*
+	 * Collision
+	 */
+
+	static FName ColliderComponentTag;
+	
+	// Override to get most relevant pawns every frame
+	virtual int32 UpdateActivePawns();
+	virtual void GatherContextsFromPawns(TArray<TWeakObjectPtr<APawn>> ActivePawns);
+	// By default getting the first SKM, if is ACharacter, get 'Mesh' component
+	// Override if e.g. The pawn is driving a boat, return SKM of boat instead
+	virtual TOptional<FShallowWaterCollisionContext> GetCollisionContextFromPawn(APawn* InPawn) const;
+	// Remove invalid references. The owning actor may get destroyed.
+	void CleanUpVehicleCollisionProxies();
+	void UpdateCollisionForPendingContexts();
+	void EnableCollisionForContext(const FShallowWaterCollisionContext& Context);
+	void DisableCollisionForContext(const FShallowWaterCollisionContext& Context);
+	// VehicleCollisionProxies is handled inside
+	void DisableCollisionForVehicle(const FShallowWaterCollisionContext& Context);
+
+	virtual float GetColliderMaxRange() const;
+	// For overriden functions: Do not use GetTheMostRelevantPlayerPawn() inside to avoid loop. Use local controlled pawn location or camera location instead
+	virtual TArray<APawn*> GetPawnsInRange(const bool bShouldSortBySignificance = false) const;
+	virtual TArray<APawn*> GetPawnsInRange(const FVector ObservingLocation, const bool bShouldSortBySignificance = false) const;
+
+	/*
+	 * WaterBody
+	 */
+	void TryGetOrWaitForWaterInfoTextureFromWaterBodies(TSet<AWaterBody*> CurrentWaterBodies);
+	UFUNCTION()
+	void OnWaterInfoTextureCreated(const UTextureRenderTarget2D* WaterInfoTexture);
+
+	UPROPERTY(BlueprintReadOnly, VisibleAnywhere, Category="Shallow Water")
+	const UTextureRenderTarget2D* WaterInfoTexture;
+	UPROPERTY(BlueprintReadOnly, VisibleAnywhere, Category="Shallow Water")
+	UTextureRenderTarget2D* NormalRT;
+
+	UPROPERTY(BlueprintReadOnly, VisibleAnywhere, Category="Shallow Water")
+	UShallowWaterSettings* Settings;
+	UPROPERTY(BlueprintReadOnly, VisibleAnywhere, Category="Shallow Water")
+	UMaterialParameterCollection* MPC;
+
+	// Consistent record through multiple frames since we want active pawns that affects the fluidsim to remain active
+	// Even when it's less significant than newly available pawns. Fluidsim enabled for different pawns each frame is bad.
+	UPROPERTY()
+	TArray<TWeakObjectPtr<APawn>> ActivePawns;
+
+	FVector PreviousProjectedLocation = FVector::ZeroVector;
+	TWeakObjectPtr<APlayerController> WeakPlayerController;
+
+private:
+	float LastTimeOverlappingAnyWaterBody = -FLT_MAX;
+	// Should only be managed by UpdateOverlappingWaterBodiesHistory()
+	TArray<TWeakObjectPtr<AWaterBody>> LastOverlappingWaterBodies_Internal;
+
+	bool bIsShallowWaterInitialized = false;
+
+	struct PendingImpact
+	{
+		FVector ImpactPosition;
+		FVector ImpactVelocity;
+		float ImpactRadius;
+	};
+
+	TArray<PendingImpact> PendingImpacts;
+	bool bFlushPendingImpactsNextTick;
+
+
+	UPROPERTY()
+	TSet<TWeakObjectPtr<UWaterBodyComponent>> WaterBodyComponentsWithProperMIDParameters;
+	UPROPERTY()
+	TArray<TWeakObjectPtr<AWaterBody>> PendingWaterBodiesToSetMIDOnInitialize;
+
+	/*
+	 * Collision Context
+	 */
+	UPROPERTY()
+	TArray<FShallowWaterCollisionContext> PreviousContexts;
+	// 'Pending' also includes contexts already with collision enabled that will be skipped
+	UPROPERTY()
+	TArray<FShallowWaterCollisionContext> PendingContexts;
+	UPROPERTY()
+	TMap<FShallowWaterCollisionContext, USkeletalMeshComponent*> VehicleCollisionProxies;
+
+	FTimerHandle WaitForPlayerControllerHandle;
+	UFUNCTION()
+	void OnLocalPlayerControllerBecomesValid(APlayerController* InPlayerController);
+	UFUNCTION()
+	void OnLocalPlayerPawnBecomesValid(APawn* OldPawn, APawn* NewPawn);
+
+	UPROPERTY()
+	TArray<FShallowWaterCollisionTracker_Actor> Tracker_Actors;
+	UPROPERTY()
+	TArray<FShallowWaterCollisionTracker_Direct> Tracker_Directs;
+
+	// Overrides collected from RegisterPhysicsAssetOverridesDataAsset
+	TMap<FGameplayTag, FShallowWaterPhysicsAssetOverride> RegisteredPhysicsAssetProxies;
+	
+	/*
+	 * Impacts
+	 */
+	int CurrentNumImpacts = 0;
+};
diff --git a/Engine/Plugins/Experimental/WaterAdvanced/Source/WaterAdvanced/Public/WaterAdvanced.h b/Engine/Plugins/Experimental/WaterAdvanced/Source/WaterAdvanced/Public/WaterAdvanced.h
new file mode 100644
--- /dev/null
+++ b/Engine/Plugins/Experimental/WaterAdvanced/Source/WaterAdvanced/Public/WaterAdvanced.h
@@ -0,0 +1,15 @@
+// Copyright Epic Games, Inc. All Rights Reserved.
+
+#pragma once
+
+#include "CoreMinimal.h"
+#include "Modules/ModuleManager.h"
+
+class FWaterAdvancedModule : public IModuleInterface
+{
+public:
+
+	/** IModuleInterface implementation */
+	virtual void StartupModule() override;
+	virtual void ShutdownModule() override;
+};
diff --git a/Engine/Plugins/Experimental/WaterAdvanced/Source/WaterAdvanced/WaterAdvanced.Build.cs b/Engine/Plugins/Experimental/WaterAdvanced/Source/WaterAdvanced/WaterAdvanced.Build.cs
new file mode 100644
--- /dev/null
+++ b/Engine/Plugins/Experimental/WaterAdvanced/Source/WaterAdvanced/WaterAdvanced.Build.cs
@@ -0,0 +1,55 @@
+// Copyright Epic Games, Inc. All Rights Reserved.
+
+using UnrealBuildTool;
+
+public class WaterAdvanced : ModuleRules
+{
+	public WaterAdvanced(ReadOnlyTargetRules Target) : base(Target)
+	{
+		PCHUsage = ModuleRules.PCHUsageMode.UseExplicitOrSharedPCHs;
+		
+		PublicIncludePaths.AddRange(
+			new string[] {
+				// ... add public include paths required here ...
+			}
+			);
+				
+		
+		PrivateIncludePaths.AddRange(
+			new string[] {
+				// ... add other private include paths required here ...
+			}
+			);
+			
+		
+		PublicDependencyModuleNames.AddRange(
+			new string[]
+			{
+				"Core"
+			}
+			);
+			
+		
+		PrivateDependencyModuleNames.AddRange(
+			new string[] {
+				"Core",
+				"CoreUObject",
+				"Engine",
+				"RenderCore",
+				"Niagara",
+				"Water",
+				"DeveloperSettings",
+				"Projects",
+				"GameplayTags"
+			}
+		);
+		
+		
+		DynamicallyLoadedModuleNames.AddRange(
+			new string[]
+			{
+				// ... add any modules that your module loads dynamically here ...
+			}
+			);
+	}
+}
diff --git a/Engine/Plugins/Experimental/WaterAdvanced/WaterAdvanced.uplugin b/Engine/Plugins/Experimental/WaterAdvanced/WaterAdvanced.uplugin
new file mode 100644
--- /dev/null
+++ b/Engine/Plugins/Experimental/WaterAdvanced/WaterAdvanced.uplugin
@@ -0,0 +1,39 @@
+{
+	"FileVersion": 3,
+	"Version": 1,
+	"VersionName": "1.0",
+	"FriendlyName": "WaterAdvanced",
+	"Description": "",
+	"Category": "Water",
+	"CreatedBy": "",
+	"CreatedByURL": "",
+	"DocsURL": "",
+	"MarketplaceURL": "",
+	"SupportURL": "",
+	"EnabledByDefault": false,
+	"CanContainContent": true,
+	"IsBetaVersion": false,
+	"IsExperimentalVersion": true,
+	"Installed": false,
+	"Modules": [
+		{
+			"Name": "WaterAdvanced",
+			"Type": "Runtime",
+			"LoadingPhase": "Default"
+		}
+	],
+	"Plugins": [
+		{
+			"Name": "Niagara",
+			"Enabled": true
+		},
+		{
+			"Name": "NiagaraFluids",
+			"Enabled": true
+		},
+		{
+			"Name": "Water",
+			"Enabled": true
+		}
+	]
+}
\ No newline at end of file
