fix initialization of ShallowWater subsystem so both the mpc and ns are loaded async

diff --git a/Engine/Plugins/Experimental/WaterAdvanced/Source/WaterAdvanced/Private/ShallowWaterSubsystem.cpp b/Engine/Plugins/Experimental/WaterAdvanced/Source/WaterAdvanced/Private/ShallowWaterSubsystem.cpp
--- a/Engine/Plugins/Experimental/WaterAdvanced/Source/WaterAdvanced/Private/ShallowWaterSubsystem.cpp
+++ b/Engine/Plugins/Experimental/WaterAdvanced/Source/WaterAdvanced/Private/ShallowWaterSubsystem.cpp
@@ -237,14 +237,32 @@ void UShallowWaterSubsystem::InitializeShallowWater()
 		return;
 	}
 
-	MPC = Settings->WaterMPC.LoadSynchronous();
-	if (MPC == nullptr)
+	// async load the ShallowWater MPC
+	if (!bAsyncLoadMPCAttempted)
 	{
-		ensure(false);
-		UE_LOG(LogShallowWater, Warning, TEXT("UShallowWaterSubsystem::InitializeShallowWater() - MPC is invalid. Make sure it's set in ShallowWater Settings."));
+		bAsyncLoadMPCAttempted = true;
+
+		Settings = GetMutableDefault<UShallowWaterSettings>();
+		UAssetManager::GetStreamableManager().RequestAsyncLoad(Settings->WaterMPC.ToSoftObjectPath(),
+			FStreamableDelegate::CreateWeakLambda(this, [this]()
+				{
+					// continue with initialization after MPC is loaded
+					InitializeShallowWater();					
+				})
+		);
+
 		return;
 	}
-	
+	else if (Settings->WaterMPC.IsValid())
+	{
+		MPC = Settings->WaterMPC.Get();
+	}
+	else
+	{
+		UE_LOG(LogShallowWater, Warning, TEXT("UShallowWaterSubsystem::InitializeShallowWater() - MPC cannot be loaded. Make sure it's set in ShallowWater Settings."));
+		return;
+	}
+
 	/*
 	 * From here, non-spectator PlayerPawn might not be available if the game is a replay
 	 * So we relies on  GetTheMostRelevantPlayerPawn();
@@ -256,13 +274,34 @@ void UShallowWaterSubsystem::InitializeShallowWater()
 		return;
 	}
 
-	UNiagaraSystem* const ShallowWaterTemplate = Settings->DefaultShallowWaterNiagaraSimulation.LoadSynchronous();
-	if (ShallowWaterTemplate == nullptr)
+	// async load the NS, then create the actor
+	UNiagaraSystem* ShallowWaterTemplate = nullptr;
+	if (!bAsyncLoadNSAttempted)
+	{
+		bAsyncLoadNSAttempted = true;
+
+		Settings = GetMutableDefault<UShallowWaterSettings>();
+		UAssetManager::GetStreamableManager().RequestAsyncLoad(Settings->DefaultShallowWaterNiagaraSimulation.ToSoftObjectPath(),
+			FStreamableDelegate::CreateWeakLambda(this, [this]()
+				{					
+					// continue with initialization after NS is loaded
+					InitializeShallowWater();
+				})
+		);
+
+		return;
+	}
+	else if (Settings->DefaultShallowWaterNiagaraSimulation.IsValid())
+	{
+		ShallowWaterTemplate = Settings->DefaultShallowWaterNiagaraSimulation.Get();
+	}
+	else
 	{
 		ensure(false);
 		UE_LOG(LogShallowWater, Warning, TEXT("UShallowWaterSubsystem::InitializeShallowWater() - Couldn't find ShallowWater template in settings"));
 		return;
 	}
+
 	const FVector SpawnLocation =  CursorPawn->GetActorLocation();
 	ShallowWaterNiagaraSimulation = UNiagaraFunctionLibrary::SpawnSystemAtLocation(GetWorld(), ShallowWaterTemplate,
 		SpawnLocation, FRotator::ZeroRotator, FVector::OneVector, false,
diff --git a/Engine/Plugins/Experimental/WaterAdvanced/Source/WaterAdvanced/Public/ShallowWaterSubsystem.h b/Engine/Plugins/Experimental/WaterAdvanced/Source/WaterAdvanced/Public/ShallowWaterSubsystem.h
--- a/Engine/Plugins/Experimental/WaterAdvanced/Source/WaterAdvanced/Public/ShallowWaterSubsystem.h
+++ b/Engine/Plugins/Experimental/WaterAdvanced/Source/WaterAdvanced/Public/ShallowWaterSubsystem.h
@@ -267,6 +267,9 @@ private:
 
 	bool bIsShallowWaterInitialized = false;
 
+	bool bAsyncLoadMPCAttempted = false;
+	bool bAsyncLoadNSAttempted = false;
+
 	struct PendingImpact
 	{
 		FVector ImpactPosition;
