Load MPC and NS async together
change raw ptrs to TObjectPtr

diff --git a/Engine/Plugins/Experimental/WaterAdvanced/Source/WaterAdvanced/Private/ShallowWaterSubsystem.cpp b/Engine/Plugins/Experimental/WaterAdvanced/Source/WaterAdvanced/Private/ShallowWaterSubsystem.cpp
--- a/Engine/Plugins/Experimental/WaterAdvanced/Source/WaterAdvanced/Private/ShallowWaterSubsystem.cpp
+++ b/Engine/Plugins/Experimental/WaterAdvanced/Source/WaterAdvanced/Private/ShallowWaterSubsystem.cpp
@@ -237,27 +237,30 @@ void UShallowWaterSubsystem::InitializeShallowWater()
 		return;
 	}
 
-	// async load the ShallowWater MPC
-	if (!bAsyncLoadMPCAttempted)
+	// async load the ShallowWater MPC and NS
+	if (!bInitializationAsyncLoadsAttempted)
 	{
-		bAsyncLoadMPCAttempted = true;
+		bInitializationAsyncLoadsAttempted = true;
 
 		Settings = GetMutableDefault<UShallowWaterSettings>();
-		UAssetManager::GetStreamableManager().RequestAsyncLoad(Settings->WaterMPC.ToSoftObjectPath(),
+
+		TArray< FSoftObjectPath> ObjectsToLoad;
+		ObjectsToLoad.Add(Settings->WaterMPC.ToSoftObjectPath());
+		ObjectsToLoad.Add(Settings->DefaultShallowWaterNiagaraSimulation.ToSoftObjectPath());
+
+		UAssetManager::GetStreamableManager().RequestAsyncLoad(ObjectsToLoad,
 			FStreamableDelegate::CreateWeakLambda(this, [this]()
 				{
-					// continue with initialization after MPC is loaded
-					InitializeShallowWater();					
+					// continue with initialization after MPC and NS are loaded
+					InitializeShallowWater();
 				})
 		);
 
 		return;
 	}
-	else if (Settings->WaterMPC.IsValid())
-	{
-		MPC = Settings->WaterMPC.Get();
-	}
-	else
+	
+	MPC = Settings->WaterMPC.Get();
+	if (MPC == nullptr)
 	{
 		UE_LOG(LogShallowWater, Warning, TEXT("UShallowWaterSubsystem::InitializeShallowWater() - MPC cannot be loaded. Make sure it's set in ShallowWater Settings."));
 		return;
@@ -275,29 +278,9 @@ void UShallowWaterSubsystem::InitializeShallowWater()
 	}
 
 	// async load the NS, then create the actor
-	UNiagaraSystem* ShallowWaterTemplate = nullptr;
-	if (!bAsyncLoadNSAttempted)
-	{
-		bAsyncLoadNSAttempted = true;
-
-		Settings = GetMutableDefault<UShallowWaterSettings>();
-		UAssetManager::GetStreamableManager().RequestAsyncLoad(Settings->DefaultShallowWaterNiagaraSimulation.ToSoftObjectPath(),
-			FStreamableDelegate::CreateWeakLambda(this, [this]()
-				{					
-					// continue with initialization after NS is loaded
-					InitializeShallowWater();
-				})
-		);
-
-		return;
-	}
-	else if (Settings->DefaultShallowWaterNiagaraSimulation.IsValid())
-	{
-		ShallowWaterTemplate = Settings->DefaultShallowWaterNiagaraSimulation.Get();
-	}
-	else
-	{
-		ensure(false);
+	UNiagaraSystem* ShallowWaterTemplate = Settings->DefaultShallowWaterNiagaraSimulation.Get();
+	if (ShallowWaterTemplate == nullptr)
+	{		
 		UE_LOG(LogShallowWater, Warning, TEXT("UShallowWaterSubsystem::InitializeShallowWater() - Couldn't find ShallowWater template in settings"));
 		return;
 	}
@@ -1159,7 +1142,7 @@ void UShallowWaterSubsystem::OnWaterInfoTextureCreated(const UTextureRenderTarge
 	WaterInfoTexture = InWaterInfoTexture;
 	if (ShallowWaterNiagaraSimulation)
 	{
-		ShallowWaterNiagaraSimulation->SetVariableTexture(FName("WaterInfoTexture"), Cast<UTexture>(const_cast<UTextureRenderTarget2D*>(WaterInfoTexture)));
+		ShallowWaterNiagaraSimulation->SetVariableTexture(FName("WaterInfoTexture"), Cast<UTexture>(const_cast<UTextureRenderTarget2D*>(WaterInfoTexture.Get())));
 	}
 }
 
diff --git a/Engine/Plugins/Experimental/WaterAdvanced/Source/WaterAdvanced/Public/ShallowWaterSubsystem.h b/Engine/Plugins/Experimental/WaterAdvanced/Source/WaterAdvanced/Public/ShallowWaterSubsystem.h
--- a/Engine/Plugins/Experimental/WaterAdvanced/Source/WaterAdvanced/Public/ShallowWaterSubsystem.h
+++ b/Engine/Plugins/Experimental/WaterAdvanced/Source/WaterAdvanced/Public/ShallowWaterSubsystem.h
@@ -243,14 +243,14 @@ protected:
 	void OnWaterInfoTextureCreated(const UTextureRenderTarget2D* WaterInfoTexture);
 
 	UPROPERTY(BlueprintReadOnly, VisibleAnywhere, Category="Shallow Water")
-	const UTextureRenderTarget2D* WaterInfoTexture;
+	TObjectPtr<const UTextureRenderTarget2D> WaterInfoTexture;
 	UPROPERTY(BlueprintReadOnly, VisibleAnywhere, Category="Shallow Water")
-	UTextureRenderTarget2D* NormalRT;
+	TObjectPtr<UTextureRenderTarget2D> NormalRT;
 
 	UPROPERTY(BlueprintReadOnly, VisibleAnywhere, Category="Shallow Water")
-	UShallowWaterSettings* Settings;
+	TObjectPtr<UShallowWaterSettings> Settings;
 	UPROPERTY(BlueprintReadOnly, VisibleAnywhere, Category="Shallow Water")
-	UMaterialParameterCollection* MPC;
+	TObjectPtr<UMaterialParameterCollection> MPC;
 
 	// Consistent record through multiple frames since we want active pawns that affects the fluidsim to remain active
 	// Even when it's less significant than newly available pawns. Fluidsim enabled for different pawns each frame is bad.
@@ -267,8 +267,7 @@ private:
 
 	bool bIsShallowWaterInitialized = false;
 
-	bool bAsyncLoadMPCAttempted = false;
-	bool bAsyncLoadNSAttempted = false;
+	bool bInitializationAsyncLoadsAttempted = false;
 
 	struct PendingImpact
 	{
