move array based impacts to use NDC

#virtualized

diff --git a/Engine/Plugins/Experimental/WaterAdvanced/Source/WaterAdvanced/Private/ShallowWaterSettings.cpp b/Engine/Plugins/Experimental/WaterAdvanced/Source/WaterAdvanced/Private/ShallowWaterSettings.cpp
--- a/Engine/Plugins/Experimental/WaterAdvanced/Source/WaterAdvanced/Private/ShallowWaterSettings.cpp
+++ b/Engine/Plugins/Experimental/WaterAdvanced/Source/WaterAdvanced/Private/ShallowWaterSettings.cpp
@@ -5,7 +5,7 @@
 UShallowWaterSettings::UShallowWaterSettings()
 {
 	DefaultShallowWaterNiagaraSimulation = TSoftObjectPtr<UNiagaraSystem>(FSoftObjectPath(TEXT("/WaterAdvanced/Niagara/Systems/Grid2D_SW_WaterBody.Grid2D_SW_WaterBody")));
-	
+	DefaultShallowWaterCollisionNDC = TSoftObjectPtr<UNiagaraDataChannelAsset>(FSoftObjectPath(TEXT("/WaterAdvanced/Niagara/DataChannels/NDC_ShallowWater.NDC_ShallowWater")));
 	WaterMPC = TSoftObjectPtr<UMaterialParameterCollection>(FSoftObjectPath(TEXT("/Water/Materials/MPC/MPC_Water.MPC_Water")));
 }
 
diff --git a/Engine/Plugins/Experimental/WaterAdvanced/Source/WaterAdvanced/Private/ShallowWaterSubsystem.cpp b/Engine/Plugins/Experimental/WaterAdvanced/Source/WaterAdvanced/Private/ShallowWaterSubsystem.cpp
--- a/Engine/Plugins/Experimental/WaterAdvanced/Source/WaterAdvanced/Private/ShallowWaterSubsystem.cpp
+++ b/Engine/Plugins/Experimental/WaterAdvanced/Source/WaterAdvanced/Private/ShallowWaterSubsystem.cpp
@@ -4,6 +4,8 @@
 
 #include "ShallowWaterSettings.h"
 #include "NiagaraComponent.h"
+#include "NiagaraDataChannel.h"
+#include "NiagaraDataChannelAccessor.h"
 #include "NiagaraDataInterfaceArrayFunctionLibrary.h"
 #include "NiagaraFunctionLibrary.h"
 #include "WaterBodyActor.h"
@@ -146,10 +148,7 @@ bool UShallowWaterSubsystem::ShouldCreateSubsystem(UObject* Outer) const
 void UShallowWaterSubsystem::Tick(float DeltaTime)
 {	
 	Super::Tick(DeltaTime);
-
-	// #todo bullet impacts should be pretty late in tick order but this can be risky
-	ResetImpacts();
-
+	
 	if (!IsShallowWaterInitialized())
 	{
 		return;
@@ -247,6 +246,7 @@ void UShallowWaterSubsystem::InitializeShallowWater()
 		TArray< FSoftObjectPath> ObjectsToLoad;
 		ObjectsToLoad.Add(Settings->WaterMPC.ToSoftObjectPath());
 		ObjectsToLoad.Add(Settings->DefaultShallowWaterNiagaraSimulation.ToSoftObjectPath());
+		ObjectsToLoad.Add(Settings->DefaultShallowWaterCollisionNDC.ToSoftObjectPath());
 
 		UAssetManager::GetStreamableManager().RequestAsyncLoad(ObjectsToLoad,
 			FStreamableDelegate::CreateWeakLambda(this, [this]()
@@ -285,6 +285,14 @@ void UShallowWaterSubsystem::InitializeShallowWater()
 		return;
 	}
 
+	// async load the NS, then create the actor
+	UNiagaraDataChannelAsset* ShallowWaterCollisionNDC = Settings->DefaultShallowWaterCollisionNDC.Get();
+	if (ShallowWaterCollisionNDC == nullptr)
+	{
+		UE_LOG(LogShallowWater, Warning, TEXT("UShallowWaterSubsystem::InitializeShallowWater() - Couldn't find ShallowWater collision NDC in settings"));
+		return;
+	}
+
 	const FVector SpawnLocation =  CursorPawn->GetActorLocation();
 	ShallowWaterNiagaraSimulation = UNiagaraFunctionLibrary::SpawnSystemAtLocation(GetWorld(), ShallowWaterTemplate,
 		SpawnLocation, FRotator::ZeroRotator, FVector::OneVector, false,
@@ -418,12 +426,7 @@ void UShallowWaterSubsystem::InitializeParameters()
 	{
 		ShallowWaterNiagaraSimulation->SetVariableVec2(FName("WorldGridSize"), FVector2D(GetGridSize()));
 		ShallowWaterNiagaraSimulation->SetVariableInt(FName("ResolutionMaxAxis"), GetGridResolution());
-		ShallowWaterNiagaraSimulation->SetVariableTextureRenderTarget(FName("NormalRT"), NormalRT);
-
-		// set variables for impacts
-		ShallowWaterNiagaraSimulation->SetVariableInt(FName("MaxImpacts"), GetMaxImpactsPerFrame());
-		ShallowWaterNiagaraSimulation->SetVariableInt(FName("NumImpacts"), 0);
-		CurrentNumImpacts = 0;
+		ShallowWaterNiagaraSimulation->SetVariableTextureRenderTarget(FName("NormalRT"), NormalRT);		
 	}
 	else
 	{
@@ -516,15 +519,6 @@ void UShallowWaterSubsystem::UpdateGridMovement()
 	PreviousProjectedLocation = ProjectedLocation;
 }
 
-void UShallowWaterSubsystem::ResetImpacts()
-{
-	if (ShallowWaterNiagaraSimulation)
-	{
-		ShallowWaterNiagaraSimulation->SetVariableInt(FName("NumImpacts"), 0);
-		CurrentNumImpacts = 0;
-	}
-}
-
 void UShallowWaterSubsystem::RegisterImpact(FVector ImpactPosition, FVector ImpactVelocity, float ImpactRadius)
 {
 	if (ShallowWaterNiagaraSimulation == nullptr)
@@ -554,7 +548,7 @@ void UShallowWaterSubsystem::RegisterImpact(FVector ImpactPosition, FVector Impa
 		return;
 	}
 
-	SetNiagaraImpactVariables(ImpactPosition, ImpactVelocity, ImpactRadius);
+	WriteImpactToNDC(ImpactPosition, ImpactVelocity, ImpactRadius);
 }
 
 // these impacts are a frame or two late, and we only want to update the Niagara system and not 
@@ -563,19 +557,22 @@ void UShallowWaterSubsystem::FlushPendingImpacts()
 {
 	for (PendingImpact CurrPendingImpact : PendingImpacts)
 	{
-		SetNiagaraImpactVariables(CurrPendingImpact.ImpactPosition, CurrPendingImpact.ImpactVelocity, CurrPendingImpact.ImpactRadius);
+		WriteImpactToNDC(CurrPendingImpact.ImpactPosition, CurrPendingImpact.ImpactVelocity, CurrPendingImpact.ImpactRadius);
 	}
 	PendingImpacts.Reset();
 }
 
-void UShallowWaterSubsystem::SetNiagaraImpactVariables(FVector ImpactPosition, FVector ImpactVelocity, float ImpactRadius)
+void UShallowWaterSubsystem::WriteImpactToNDC(FVector ImpactPosition, FVector ImpactVelocity, float ImpactRadius)
 {
-	UNiagaraDataInterfaceArrayFunctionLibrary::SetNiagaraArrayPositionValue(ShallowWaterNiagaraSimulation, "User.ImpactPositionArray", CurrentNumImpacts, ImpactPosition, false);
-	UNiagaraDataInterfaceArrayFunctionLibrary::SetNiagaraArrayVectorValue(ShallowWaterNiagaraSimulation, "User.ImpactVelocityArray", CurrentNumImpacts, ImpactVelocity, false);
-	UNiagaraDataInterfaceArrayFunctionLibrary::SetNiagaraArrayFloatValue(ShallowWaterNiagaraSimulation, "User.ImpactRadiusArray", CurrentNumImpacts, ImpactRadius, false);
-
-	++CurrentNumImpacts;
-	ShallowWaterNiagaraSimulation->SetVariableInt(FName("NumImpacts"), CurrentNumImpacts);
+	FNiagaraDataChannelSearchParameters SearchParams(ImpactPosition);
+	const UNiagaraDataChannelAsset* NDC = Settings->DefaultShallowWaterCollisionNDC.Get();
+	if (UNiagaraDataChannelWriter* DCWriter = UNiagaraDataChannelLibrary::WriteToNiagaraDataChannel(ShallowWaterNiagaraSimulation, NDC, SearchParams, 1, false, true, true, TEXT("ShallowWaterWriteImpact")))
+	{
+		int32 Index = 0;
+		DCWriter->WritePosition(TEXT("Position"), Index, ImpactPosition);
+		DCWriter->WriteVector(TEXT("Velocity"), Index, ImpactVelocity);
+		DCWriter->WriteFloat(TEXT("Radius"), Index, ImpactRadius);
+	}
 }
 
 void UShallowWaterSubsystem::SetWaterBodyMIDParameters(AWaterBody* WaterBody)
diff --git a/Engine/Plugins/Experimental/WaterAdvanced/Source/WaterAdvanced/Public/ShallowWaterCommon.h b/Engine/Plugins/Experimental/WaterAdvanced/Source/WaterAdvanced/Public/ShallowWaterCommon.h
--- a/Engine/Plugins/Experimental/WaterAdvanced/Source/WaterAdvanced/Public/ShallowWaterCommon.h
+++ b/Engine/Plugins/Experimental/WaterAdvanced/Source/WaterAdvanced/Public/ShallowWaterCommon.h
@@ -17,9 +17,7 @@ struct WATERADVANCED_API FShallowWaterSimParameters
 	UPROPERTY(EditAnywhere, Category = "Shallow Water")
 	float WorldGridSize = 2000.f;
 	UPROPERTY(EditAnywhere, Category = "Shallow Water")
-	int32 ResolutionMaxAxis = 512;
-	UPROPERTY(EditAnywhere, Category = "Shallow Water")
-	int32 MaxImpactsPerFrame = 50;
+	int32 ResolutionMaxAxis = 512;	
 };
 
 USTRUCT(BlueprintType)
diff --git a/Engine/Plugins/Experimental/WaterAdvanced/Source/WaterAdvanced/Public/ShallowWaterSettings.h b/Engine/Plugins/Experimental/WaterAdvanced/Source/WaterAdvanced/Public/ShallowWaterSettings.h
--- a/Engine/Plugins/Experimental/WaterAdvanced/Source/WaterAdvanced/Public/ShallowWaterSettings.h
+++ b/Engine/Plugins/Experimental/WaterAdvanced/Source/WaterAdvanced/Public/ShallowWaterSettings.h
@@ -5,6 +5,7 @@
 #include "CoreMinimal.h"
 #include "Engine/DeveloperSettings.h"
 #include "NiagaraSystem.h"
+#include "NiagaraDataChannelPublic.h"
 #include "ShallowWaterCommon.h"
 #include "ShallowWaterPhysicsAssetOverridesDataAsset.h"
 #include "Materials/MaterialParameterCollection.h"
@@ -21,6 +22,8 @@ public:
 	
 	UPROPERTY(Config, BlueprintReadOnly, EditAnywhere, Category="Shallow Water Simulation")
 	TSoftObjectPtr<UNiagaraSystem> DefaultShallowWaterNiagaraSimulation;
+	UPROPERTY(Config, BlueprintReadOnly, EditAnywhere, Category = "Shallow Water Simulation")
+	TSoftObjectPtr<UNiagaraDataChannelAsset> DefaultShallowWaterCollisionNDC;
 	UPROPERTY(Config, BlueprintReadOnly, EditAnywhere, Category="Shallow Water Simulation")
 	FShallowWaterSimParameters ShallowWaterSimParameters;
 	UPROPERTY(Config, BlueprintReadOnly, EditAnywhere, Category="Shallow Water Simulation")
diff --git a/Engine/Plugins/Experimental/WaterAdvanced/Source/WaterAdvanced/Public/ShallowWaterSubsystem.h b/Engine/Plugins/Experimental/WaterAdvanced/Source/WaterAdvanced/Public/ShallowWaterSubsystem.h
--- a/Engine/Plugins/Experimental/WaterAdvanced/Source/WaterAdvanced/Public/ShallowWaterSubsystem.h
+++ b/Engine/Plugins/Experimental/WaterAdvanced/Source/WaterAdvanced/Public/ShallowWaterSubsystem.h
@@ -155,14 +155,11 @@ public:
 	 */
 	void UpdateGridMovement();
 
-	// #note - these are nonblocking for now
-	UFUNCTION(BlueprintCallable, Category = "Shallow Water")
-	void ResetImpacts();
 	UFUNCTION(BlueprintCallable, Category = "Shallow Water")
 	void RegisterImpact(FVector ImpactPosition, FVector ImpactVelocity, float ImpactRadius);
 	
 	void FlushPendingImpacts();
-	void SetNiagaraImpactVariables(FVector ImpactPosition, FVector ImpactVelocity, float ImpactRadius);
+	void WriteImpactToNDC(FVector ImpactPosition, FVector ImpactVelocity, float ImpactRadius);
 	// Override to return the ECC channel of bullets
 	virtual ECollisionChannel GetImpactCollisionChannel() PURE_VIRTUAL(UShallowWaterSubsystem::GetImpactCollisionChannel, return ECC_WorldDynamic;)
 	
@@ -173,8 +170,7 @@ public:
 
 	void TryUpdateWaterBodyMIDParameters(UWaterBodyComponent* WaterBodyComponent);
 	float GetGridSize() const { return Settings->ShallowWaterSimParameters.WorldGridSize; }
-	int32 GetGridResolution() const { return Settings->ShallowWaterSimParameters.ResolutionMaxAxis; }
-	int32 GetMaxImpactsPerFrame() const { return Settings->ShallowWaterSimParameters.MaxImpactsPerFrame; }
+	int32 GetGridResolution() const { return Settings->ShallowWaterSimParameters.ResolutionMaxAxis; }	
 
 	/*
 	 * Add PA overrides. Designed to be called by Game Feature Plugins. 
@@ -309,9 +305,4 @@ private:
 
 	// Overrides collected from RegisterPhysicsAssetOverridesDataAsset
 	TMap<FGameplayTag, FShallowWaterPhysicsAssetOverride> RegisteredPhysicsAssetProxies;
-	
-	/*
-	 * Impacts
-	 */
-	int CurrentNumImpacts = 0;
 };
