ShallowWaterRiverActor - Actor for building shallow water based rivers off of a list of spline based Water Body rivers

#changelist validated
#virtualized

diff --git a/Engine/Plugins/Experimental/WaterAdvanced/Source/WaterAdvanced/Private/ShallowWaterRiverActor.cpp b/Engine/Plugins/Experimental/WaterAdvanced/Source/WaterAdvanced/Private/ShallowWaterRiverActor.cpp
new file mode 100644
--- /dev/null
+++ b/Engine/Plugins/Experimental/WaterAdvanced/Source/WaterAdvanced/Private/ShallowWaterRiverActor.cpp
@@ -0,0 +1,294 @@
+// Copyright Epic Games, Inc. All Rights Reserved.
+
+#include "ShallowWaterRiverActor.h"
+#include "NiagaraComponent.h"
+#include "NiagaraSystem.h"
+#include "NiagaraFunctionLibrary.h"
+#include "WaterBodyActor.h"
+#include "WaterSplineComponent.h"
+#include "Engine/TextureRenderTarget2D.h"
+#include "ShallowWaterCommon.h"
+#include "Engine/World.h"
+
+UShallowWaterRiverComponent::UShallowWaterRiverComponent(const FObjectInitializer& Initializer)
+	: Super(Initializer)
+{
+	PrimaryComponentTick.bCanEverTick = true;	
+
+#if WITH_EDITORONLY_DATA
+	bTickInEditor = true;
+#endif // WITH_EDITORONLY_DATA
+
+	bIsInitialized = false;
+	bTickInitialize = false;
+}
+
+void UShallowWaterRiverComponent::PostLoad()
+{
+	Super::PostLoad();
+
+#if WITH_EDITOR
+	bIsInitialized = false;
+	bTickInitialize = false;
+
+	Rebuild();
+#endif
+}
+
+void UShallowWaterRiverComponent::TickComponent(float DeltaTime, ELevelTick TickType, FActorComponentTickFunction* ThisTickFunction)
+{
+#if WITH_EDITOR
+	// lots of tick ordering issues, so we try to initialize on the first tick too
+	if (!bIsInitialized && !bTickInitialize)
+	{
+		bTickInitialize = true;
+		Rebuild();
+	}
+
+#endif
+}
+
+#if WITH_EDITOR
+
+void UShallowWaterRiverComponent::PostEditChangeProperty(FPropertyChangedEvent& PropertyChangedEvent)
+{
+	Super::PostEditChangeProperty(PropertyChangedEvent);
+
+	Rebuild();
+}
+
+void UShallowWaterRiverComponent::Rebuild()
+{	
+	if (RiverSimSystem != nullptr)
+	{
+		RiverSimSystem->SetActive(false);
+		RiverSimSystem->DestroyComponent();
+		RiverSimSystem = nullptr;
+	}
+			
+	// collect all the water bodies
+	TSet < TObjectPtr<AWaterBody>> AllWaterBodies;
+
+	if (SourceRiverWaterBody != nullptr)
+	{
+		AllWaterBodies.Add(SourceRiverWaterBody);
+	}
+	else	
+	{
+		UE_LOG(LogShallowWater, Warning, TEXT("UShallowWaterRiverComponent::Rebuild() - No source water body specified"));
+		return;
+	}
+	
+	if (SinkRiverWaterBody != nullptr)
+	{
+		AllWaterBodies.Add(SinkRiverWaterBody);
+	}	
+
+	for (TObjectPtr<AWaterBody > CurrWaterBody : AdditonalRiverWaterBodies)
+	{
+		AllWaterBodies.Add(CurrWaterBody);
+	}
+
+	if (AllWaterBodies.Num() == 0)
+	{
+		UE_LOG(LogShallowWater, Warning, TEXT("UShallowWaterRiverComponent::Rebuild() - No water bodies specified"));
+		return;
+	}
+
+	// accumulate bounding box for river water bodies
+	FBoxSphereBounds::Builder CombinedWorldBoundsBuilder;
+	for (TObjectPtr<AWaterBody > CurrWaterBody : AllWaterBodies)
+	{
+		if (CurrWaterBody != nullptr)
+		{
+			// accumulate bounds
+			FBoxSphereBounds WorldBounds;
+			CurrWaterBody->GetActorBounds(false, WorldBounds.Origin, WorldBounds.BoxExtent);
+
+			CombinedWorldBoundsBuilder += WorldBounds;
+		}
+		else
+		{
+			UE_LOG(LogShallowWater, Warning, TEXT("UShallowWaterRiverComponent::Rebuild() - null water body actor found"));
+			return;
+		}
+	}
+	FBoxSphereBounds CombinedBounds(CombinedWorldBoundsBuilder);
+
+	if (CombinedBounds.BoxExtent.Length() < SMALL_NUMBER)
+	{
+		UE_LOG(LogShallowWater, Warning, TEXT("UShallowWaterRiverComponent::Rebuild() - river bodies have zero bounds"));
+		return;
+	}
+
+	FVector SourcePos(0, 0, 0);
+	float SourceWidth = 1;
+	float SourceDepth = 1;
+	FVector SourceDir(1, 0, 0);	
+
+	// Get source
+	if (!QueryWaterAtSplinePoint(SourceRiverWaterBody, 0, SourcePos, SourceDir, SourceWidth, SourceDepth))
+	{
+		UE_LOG(LogShallowWater, Warning, TEXT("UShallowWaterRiverComponent::Rebuild() - water source query failed"));
+		return;
+	}
+
+	FVector SinkPos(0, 0, 0);
+	float SinkWidth = 1;
+	float SinkDepth = 1;
+	FVector SinkDir(1, 0, 0);
+
+	// Get Sink		
+	TObjectPtr<AWaterBody> SinkToUse = SinkRiverWaterBody;
+
+	// if no sink specified, use source
+	if (SinkToUse == nullptr)
+	{
+		SinkToUse = SourceRiverWaterBody;
+	}
+	
+	if (!QueryWaterAtSplinePoint(SinkToUse, -1, SinkPos, SinkDir, SinkWidth, SinkDepth))
+	{
+		UE_LOG(LogShallowWater, Warning, TEXT("UShallowWaterRiverComponent::Rebuild() - water sink query failed"));
+		return;
+	}
+	
+
+	/*
+				bool WaterInfoTextureSet = false;
+
+				if (!WaterInfoTextureSet)
+				{
+					AWaterZone* WaterZone = CurrWaterBodyComponent->GetWaterZone();
+					if (WaterZone)
+					{
+						const TObjectPtr<UTextureRenderTarget2D> NewWaterInfoTexture = WaterZone->WaterInfoTexture;
+
+						if (NewWaterInfoTexture == nullptr)
+						{
+							WaterZone->GetOnWaterInfoTextureCreated().RemoveDynamic(this, &UShallowWaterRiverComponent::OnWaterInfoTextureCreated);
+							WaterZone->GetOnWaterInfoTextureCreated().AddDynamic(this, &UShallowWaterRiverComponent::OnWaterInfoTextureCreated);
+						}
+						else
+						{
+							OnWaterInfoTextureCreated(NewWaterInfoTexture);
+						}
+
+						WaterInfoTextureSet = true;
+					}
+				}							
+	}
+	*/	
+
+	// spawn niagara system
+	FVector SystemPos = CombinedBounds.Origin - FVector(0, 0, CombinedBounds.BoxExtent.Z);	
+	RiverSimSystem = UNiagaraFunctionLibrary::SpawnSystemAtLocation(GetWorld(), NiagaraRiverSimulation, SystemPos,
+		FRotator::ZeroRotator, FVector::OneVector, false, true, ENCPoolMethod::None, false);
+
+	if (RiverSimSystem == nullptr)
+	{
+		UE_LOG(LogShallowWater, Warning, TEXT("UShallowWaterRiverComponent::Rebuild() - Cannot spawn river system"));
+		return;
+	}
+
+	RiverSimSystem->Activate();
+
+	RiverSimSystem->SetVariableVec2(FName("WorldGridSize"), 2.0f * FVector2D(CombinedBounds.BoxExtent.X, CombinedBounds.BoxExtent.Y));
+
+	// pad out source's box height a so it intersects the sim plane.  This value doesn't matter much so we hardcode it
+	float Overshoot = 1000.f;
+	float FinalSourceHeight = 2. * CombinedBounds.BoxExtent.Z + Overshoot;
+
+	RiverSimSystem->SetVariablePosition(FName("SourcePos"), SourcePos - FVector(0, 0, .5 * FinalSourceHeight) + FVector(SourceDir.X, SourceDir.Y, 0) * .5 * SourceSize);
+	RiverSimSystem->SetVariableVec3(FName("SourceSize"), FVector(SourceWidth, SourceSize, FinalSourceHeight));
+	RiverSimSystem->SetVariableFloat(FName("SourceAngle"), PI / 2.f + FMath::Acos(SourceDir.Dot(FVector(1,0,0))));
+	
+	// height of the sink box doesn't matter
+	float SinkBoxHeight = 10000000;
+	RiverSimSystem->SetVariablePosition(FName("SinkPos"), SinkPos);
+	RiverSimSystem->SetVariableVec3(FName("SinkSize"), FVector(SinkWidth, SourceSize, SinkBoxHeight));
+	RiverSimSystem->SetVariableFloat(FName("SinkAngle"), PI / 2.f + FMath::Acos(SinkDir.Dot(FVector(1, 0, 0))));
+
+	bIsInitialized = true;
+}
+
+void UShallowWaterRiverComponent::OnWaterInfoTextureCreated(const UTextureRenderTarget2D* InWaterInfoTexture)
+{
+	if (InWaterInfoTexture == nullptr)
+	{
+		ensureMsgf(false, TEXT("UShallowWaterRiverComponent::OnWaterInfoTextureCreated was called with NULL WaterInfoTexture"));
+		return;
+	}
+
+	WaterInfoTexture = InWaterInfoTexture;
+	if (RiverSimSystem)
+	{
+		RiverSimSystem->SetVariableTexture(FName("WaterInfoTexture"), Cast<UTexture>(const_cast<UTextureRenderTarget2D*>(WaterInfoTexture.Get())));
+	}
+	else
+	{
+		ensureMsgf(false, TEXT("UShallowWaterRiverComponent::OnWaterInfoTextureCreated No river system to set water info on"));
+	}
+}
+
+bool UShallowWaterRiverComponent::QueryWaterAtSplinePoint(TObjectPtr<AWaterBody> WaterBody, int SplinePoint, FVector& OutPos, FVector& OutTangent, float& OutWidth, float& OutDepth)
+{	
+	if (WaterBody != nullptr)
+	{
+		TObjectPtr<UWaterBodyComponent> CurrWaterBodyComponent = WaterBody->GetWaterBodyComponent();
+
+		UWaterSplineComponent* CurrSpline = WaterBody->GetWaterSpline();
+		
+		if (CurrSpline != nullptr)
+		{
+			// -1 means last spline point
+			if (SplinePoint == -1)
+			{
+				SplinePoint = CurrSpline->GetNumberOfSplinePoints() - 1;
+			}
+
+			UWaterSplineMetadata* Metadata = WaterBody->GetWaterSplineMetadata();
+
+			if (Metadata != nullptr)
+			{
+				OutPos = CurrSpline->SplineCurves.Position.Points[SplinePoint].OutVal;
+				OutPos = WaterBody->GetActorTransform().TransformPosition(OutPos);
+
+				OutWidth = Metadata->RiverWidth.Points[SplinePoint].OutVal;
+				OutDepth = Metadata->Depth.Points[SplinePoint].OutVal;
+
+				OutTangent = CurrSpline->SplineCurves.Position.Points[SplinePoint].LeaveTangent;
+				OutTangent.Normalize();
+			}
+			else
+			{
+				UE_LOG(LogShallowWater, Warning, TEXT("UShallowWaterRiverComponent::QueryWaterAtSplinePoint() - Water spline metadata is null"));
+				return false;
+			}
+		}
+		else
+		{
+			UE_LOG(LogShallowWater, Warning, TEXT("UShallowWaterRiverComponent::QueryWaterAtSplinePoint() - Water spline component is null"));
+			return false;
+		}
+	}
+	else
+	{
+		UE_LOG(LogShallowWater, Warning, TEXT("UShallowWaterRiverComponent::QueryWaterAtSplinePoint() - Water actor is null"));
+		return false;
+	}
+
+	return true;
+}
+#endif
+
+AShallowWaterRiver::AShallowWaterRiver(const FObjectInitializer& ObjectInitializer)
+	: Super(ObjectInitializer)
+{
+	ShallowWaterRiverComponent = CreateDefaultSubobject<UShallowWaterRiverComponent>(TEXT("ShallowWaterRiverComponent"));
+	RootComponent = ShallowWaterRiverComponent;
+
+	PrimaryActorTick.bCanEverTick = true;
+	SetHidden(false);
+}
+
diff --git a/Engine/Plugins/Experimental/WaterAdvanced/Source/WaterAdvanced/Public/ShallowWaterRiverActor.h b/Engine/Plugins/Experimental/WaterAdvanced/Source/WaterAdvanced/Public/ShallowWaterRiverActor.h
new file mode 100644
--- /dev/null
+++ b/Engine/Plugins/Experimental/WaterAdvanced/Source/WaterAdvanced/Public/ShallowWaterRiverActor.h
@@ -0,0 +1,77 @@
+// Copyright Epic Games, Inc. All Rights Reserved.
+
+#pragma once
+
+#include "CoreMinimal.h"
+#include "GameFramework/Actor.h"
+#include "Components/PrimitiveComponent.h"
+#include "ShallowWaterRiverActor.generated.h"
+
+class UNiagaraComponent;
+class UNiagaraSystem;
+class AWaterBody;
+class UTextureRenderTarget2D;
+
+UCLASS(BlueprintType, HideCategories = (Physics, Replication, Input, Collision))
+class WATERADVANCED_API UShallowWaterRiverComponent : public UPrimitiveComponent
+{
+	GENERATED_UCLASS_BODY()
+
+public:
+	UPROPERTY(EditAnywhere, Category = "Simulation", meta = (DisplayName = "Niagara River Simulation"))
+	UNiagaraSystem *NiagaraRiverSimulation;
+
+	UPROPERTY(EditAnywhere, Category = "Simulation", meta = (DisplayName = "Resolution Max Axis"))
+	int ResolutionMaxAxis;
+
+	UPROPERTY(EditAnywhere, Category = "Simulation", meta = (DisplayName = "Source Width"))
+	float SourceSize;
+
+	UPROPERTY(EditAnywhere, Category = "Water", meta = (DisplayName = "Source River Water Body"))
+	TObjectPtr<AWaterBody> SourceRiverWaterBody;
+
+	UPROPERTY(EditAnywhere, Category = "Water", meta = (DisplayName = "Sink River Water Body"))
+	TObjectPtr<AWaterBody> SinkRiverWaterBody;
+
+	UPROPERTY(EditAnywhere, Category = "Water", meta = (DisplayName = "Additional River Water Bodies"))
+	TArray<TObjectPtr<AWaterBody>> AdditonalRiverWaterBodies;
+
+	UPROPERTY(EditAnywhere, Category = "Reset", meta = (DisplayName = "Reset"))
+	bool Reset;
+
+	virtual void PostLoad() override;
+
+	virtual void TickComponent(float DeltaTime, ELevelTick TickType, FActorComponentTickFunction* ThisTickFunction) override;
+
+#if WITH_EDITOR
+	void Rebuild();
+
+	virtual void PostEditChangeProperty(FPropertyChangedEvent& PropertyChangedEvent) override;
+	void OnWaterInfoTextureCreated(const UTextureRenderTarget2D* InWaterInfoTexture);
+#endif // WITH_EDITOR
+
+protected:
+	// Asset can be set in Project Settings - Plugins - Water ShallowWaterSimulation
+	UPROPERTY(BlueprintReadOnly, VisibleDefaultsOnly, Category = "Shallow Water")
+	TObjectPtr<UNiagaraComponent> RiverSimSystem;
+
+	UPROPERTY(BlueprintReadOnly, VisibleDefaultsOnly, Category = "Shallow Water")
+	TObjectPtr<const UTextureRenderTarget2D> WaterInfoTexture;
+
+	bool QueryWaterAtSplinePoint(TObjectPtr<AWaterBody> WaterBody, int SplinePoint, FVector& OutPos, FVector& OutTangent, float& OutWidth, float& OutDepth);
+
+private:
+	bool bIsInitialized;	
+	bool bTickInitialize;
+};
+
+UCLASS(BlueprintType, HideCategories = (Physics, Replication, Input, Collision))
+class WATERADVANCED_API AShallowWaterRiver : public AActor
+{
+	GENERATED_UCLASS_BODY()
+
+private:
+	// Asset can be set in Project Settings - Plugins - Water ShallowWaterSimulation
+	UPROPERTY(VisibleAnywhere, Category = "Shallow Water")
+	TObjectPtr<UShallowWaterRiverComponent> ShallowWaterRiverComponent;	
+};
\ No newline at end of file
