update component ownership, better defaults, default river system better initialization
template river system, updated assets

#changelist validated
#virtualized

diff --git a/Engine/Plugins/Experimental/WaterAdvanced/Source/WaterAdvanced/Private/ShallowWaterRiverActor.cpp b/Engine/Plugins/Experimental/WaterAdvanced/Source/WaterAdvanced/Private/ShallowWaterRiverActor.cpp
--- a/Engine/Plugins/Experimental/WaterAdvanced/Source/WaterAdvanced/Private/ShallowWaterRiverActor.cpp
+++ b/Engine/Plugins/Experimental/WaterAdvanced/Source/WaterAdvanced/Private/ShallowWaterRiverActor.cpp
@@ -21,6 +21,10 @@ UShallowWaterRiverComponent::UShallowWaterRiverComponent(const FObjectInitialize
 
 	bIsInitialized = false;
 	bTickInitialize = false;
+
+	ResolutionMaxAxis = 512;
+	SourceSize = 1000;
+	NiagaraRiverSimulation = LoadObject<UNiagaraSystem>(nullptr, TEXT("/WaterAdvanced/Niagara/Systems/Grid2D_SW_River.Grid2D_SW_River"));
 }
 
 void UShallowWaterRiverComponent::PostLoad()
@@ -48,6 +52,18 @@ void UShallowWaterRiverComponent::TickComponent(float DeltaTime, ELevelTick Tick
 #endif
 }
 
+void UShallowWaterRiverComponent::OnUnregister()
+{
+	Super::OnUnregister();
+
+	//if (RiverSimSystem && !RiverSimSystem->IsBeingDestroyed())
+	//{
+	//	RiverSimSystem->SetActive(false);
+	//	RiverSimSystem->DestroyComponent();
+	//	RiverSimSystem = nullptr;
+	//}
+}
+
 #if WITH_EDITOR
 
 void UShallowWaterRiverComponent::PostEditChangeProperty(FPropertyChangedEvent& PropertyChangedEvent)
@@ -180,10 +196,30 @@ void UShallowWaterRiverComponent::Rebuild()
 	}
 	*/	
 
-	// spawn niagara system
-	FVector SystemPos = CombinedBounds.Origin - FVector(0, 0, CombinedBounds.BoxExtent.Z);	
-	RiverSimSystem = UNiagaraFunctionLibrary::SpawnSystemAtLocation(GetWorld(), NiagaraRiverSimulation, SystemPos,
-		FRotator::ZeroRotator, FVector::OneVector, false, true, ENCPoolMethod::None, false);
+
+	FVector SystemPos = CombinedBounds.Origin - FVector(0, 0, CombinedBounds.BoxExtent.Z);
+	
+	RiverSimSystem = NewObject<UNiagaraComponent>(this, NAME_None, RF_Transient);
+	RiverSimSystem->bUseAttachParentBound = false;
+	RiverSimSystem->SetWorldLocation(SystemPos);
+
+	if (GetWorld() && GetWorld()->bIsWorldInitialized)
+	{
+		if (!RiverSimSystem->IsRegistered())
+		{
+			RiverSimSystem->RegisterComponentWithWorld(GetWorld());
+		}
+
+		RiverSimSystem->SetVisibleFlag(true);
+		RiverSimSystem->SetAsset(NiagaraRiverSimulation);
+		RiverSimSystem->ReinitializeSystem();
+	}
+	else
+	{
+		UE_LOG(LogShallowWater, Warning, TEXT("UShallowWaterRiverComponent::Rebuild() - World not initialized"));
+		return;
+	}
+	
 
 	if (RiverSimSystem == nullptr)
 	{
diff --git a/Engine/Plugins/Experimental/WaterAdvanced/Source/WaterAdvanced/Public/ShallowWaterRiverActor.h b/Engine/Plugins/Experimental/WaterAdvanced/Source/WaterAdvanced/Public/ShallowWaterRiverActor.h
--- a/Engine/Plugins/Experimental/WaterAdvanced/Source/WaterAdvanced/Public/ShallowWaterRiverActor.h
+++ b/Engine/Plugins/Experimental/WaterAdvanced/Source/WaterAdvanced/Public/ShallowWaterRiverActor.h
@@ -19,7 +19,7 @@ class WATERADVANCED_API UShallowWaterRiverComponent : public UPrimitiveComponent
 
 public:
 	UPROPERTY(EditAnywhere, Category = "Simulation", meta = (DisplayName = "Niagara River Simulation"))
-	UNiagaraSystem *NiagaraRiverSimulation;
+	TObjectPtr <class UNiagaraSystem> NiagaraRiverSimulation;
 
 	UPROPERTY(EditAnywhere, Category = "Simulation", meta = (DisplayName = "Resolution Max Axis"))
 	int ResolutionMaxAxis;
@@ -43,6 +43,8 @@ public:
 
 	virtual void TickComponent(float DeltaTime, ELevelTick TickType, FActorComponentTickFunction* ThisTickFunction) override;
 
+	virtual void OnUnregister() override;
+
 #if WITH_EDITOR
 	void Rebuild();
 
