remove default for Niagara System in river component

diff --git a/Engine/Plugins/Experimental/WaterAdvanced/Source/WaterAdvanced/Private/ShallowWaterRiverActor.cpp b/Engine/Plugins/Experimental/WaterAdvanced/Source/WaterAdvanced/Private/ShallowWaterRiverActor.cpp
--- a/Engine/Plugins/Experimental/WaterAdvanced/Source/WaterAdvanced/Private/ShallowWaterRiverActor.cpp
+++ b/Engine/Plugins/Experimental/WaterAdvanced/Source/WaterAdvanced/Private/ShallowWaterRiverActor.cpp
@@ -24,7 +24,7 @@ UShallowWaterRiverComponent::UShallowWaterRiverComponent(const FObjectInitialize
 
 	ResolutionMaxAxis = 512;
 	SourceSize = 1000;
-	NiagaraRiverSimulation = LoadObject<UNiagaraSystem>(nullptr, TEXT("/WaterAdvanced/Niagara/Systems/Grid2D_SW_River.Grid2D_SW_River"));
+	//NiagaraRiverSimulation = LoadObject<UNiagaraSystem>(nullptr, TEXT("/WaterAdvanced/Niagara/Systems/Grid2D_SW_River.Grid2D_SW_River"));
 }
 
 void UShallowWaterRiverComponent::PostLoad()
@@ -81,7 +81,12 @@ void UShallowWaterRiverComponent::Rebuild()
 		RiverSimSystem->DestroyComponent();
 		RiverSimSystem = nullptr;
 	}
-			
+	
+	if (NiagaraRiverSimulation == nullptr)
+	{
+		UE_LOG(LogShallowWater, Warning, TEXT("UShallowWaterRiverComponent::Rebuild() - null Niagara system asset"));
+	}
+
 	// collect all the water bodies
 	TSet < TObjectPtr<AWaterBody>> AllWaterBodies;
 
