support toggle for caching
buttons on actor for reset and caching

#changelist validated
#virtualized

diff --git a/Engine/Plugins/Experimental/WaterAdvanced/Source/WaterAdvanced/Private/ShallowWaterRiverActor.cpp b/Engine/Plugins/Experimental/WaterAdvanced/Source/WaterAdvanced/Private/ShallowWaterRiverActor.cpp
--- a/Engine/Plugins/Experimental/WaterAdvanced/Source/WaterAdvanced/Private/ShallowWaterRiverActor.cpp
+++ b/Engine/Plugins/Experimental/WaterAdvanced/Source/WaterAdvanced/Private/ShallowWaterRiverActor.cpp
@@ -9,6 +9,7 @@
 #include "Engine/TextureRenderTarget2D.h"
 #include "ShallowWaterCommon.h"
 #include "Engine/World.h"
+#include "Engine/Texture2D.h"
 
 UShallowWaterRiverComponent::UShallowWaterRiverComponent(const FObjectInitializer& Initializer)
 	: Super(Initializer)
@@ -70,7 +71,21 @@ void UShallowWaterRiverComponent::PostEditChangeProperty(FPropertyChangedEvent&
 {
 	Super::PostEditChangeProperty(PropertyChangedEvent);
 
-	Rebuild();
+	FName PropertyName;
+	if (PropertyChangedEvent.Property)
+	{
+		PropertyName = PropertyChangedEvent.Property->GetFName();
+	}
+			
+	// this should go before rebuild not after...something is wrong
+	if (PropertyName == GET_MEMBER_NAME_CHECKED(UShallowWaterRiverComponent, PreviewBakedSim))
+	{
+		RiverSimSystem->SetVariableBool(FName("ReadCachedSim"), PreviewBakedSim);
+	}
+	else
+	{
+		Rebuild();
+	}
 }
 
 void UShallowWaterRiverComponent::Rebuild()
@@ -250,9 +265,31 @@ void UShallowWaterRiverComponent::Rebuild()
 	RiverSimSystem->SetVariableVec3(FName("SinkSize"), FVector(SinkWidth, SourceSize, SinkBoxHeight));
 	RiverSimSystem->SetVariableFloat(FName("SinkAngle"), PI / 2.f + FMath::Acos(SinkDir.Dot(FVector(1, 0, 0))));
 
+	RiverSimSystem->SetVariableFloat(FName("SimSpeed"), SimSpeed);
+	RiverSimSystem->SetVariableInt(FName("NumSteps"), NumSteps);	
+	
+	BakedWaterSurfaceRT = NewObject<UTextureRenderTarget2D>(this, NAME_None, RF_Transient);
+	BakedWaterSurfaceRT->InitAutoFormat(1, 1);
+	RiverSimSystem->SetVariableTextureRenderTarget(FName("SimGridRT"), BakedWaterSurfaceRT);
+	RiverSimSystem->SetVariableBool(FName("ReadCachedSim"), PreviewBakedSim);
+
+	if (BakedWaterSurfaceTexture != nullptr)
+	{
+		RiverSimSystem->SetVariableTexture(FName("BakedSimTexture"), BakedWaterSurfaceTexture);
+	}
+
 	bIsInitialized = true;
 }
 
+void UShallowWaterRiverComponent::Bake()
+{	
+	EObjectFlags TextureObjectFlags = EObjectFlags::RF_Public;
+
+	BakedWaterSurfaceTexture = BakedWaterSurfaceRT->ConstructTexture2D(this, "BakedRiverTexture", TextureObjectFlags);
+
+	RiverSimSystem->SetVariableTexture(FName("BakedSimTexture"), BakedWaterSurfaceTexture);
+}
+
 void UShallowWaterRiverComponent::OnWaterInfoTextureCreated(const UTextureRenderTarget2D* InWaterInfoTexture)
 {
 	if (InWaterInfoTexture == nullptr)
diff --git a/Engine/Plugins/Experimental/WaterAdvanced/Source/WaterAdvanced/Private/ShallowWaterRiverDetails.cpp b/Engine/Plugins/Experimental/WaterAdvanced/Source/WaterAdvanced/Private/ShallowWaterRiverDetails.cpp
new file mode 100644
--- /dev/null
+++ b/Engine/Plugins/Experimental/WaterAdvanced/Source/WaterAdvanced/Private/ShallowWaterRiverDetails.cpp
@@ -0,0 +1,189 @@
+// Copyright Epic Games, Inc. All Rights Reserved.
+
+#if WITH_EDITOR
+#include "ShallowWaterRiverDetails.h"
+#include "PropertyHandle.h"
+#include "DetailLayoutBuilder.h"
+#include "DetailCategoryBuilder.h"
+#include "DetailWidgetRow.h"
+#include "PropertyCustomizationHelpers.h"
+#include "Widgets/Text/STextBlock.h"
+#include "Widgets/Layout/SBox.h"
+#include "Styling/AppStyle.h"
+#include "PropertyEditorModule.h"
+#include "Modules/ModuleManager.h"
+#include "Widgets/Input/SButton.h"
+#include "Widgets/Layout/SUniformGridPanel.h"
+#include "GameDelegates.h"
+#include "IAssetTools.h"
+#include "IDetailChildrenBuilder.h"
+#include "ShallowWaterRiverActor.h"
+#include "Engine/World.h"
+#include "CoreMinimal.h"
+#include "UObject/ObjectMacros.h"
+#include "UObject/UnrealType.h"
+
+#define LOCTEXT_NAMESPACE "ShallowWaterRiverDetails"
+
+TSharedRef<IDetailCustomization> FShallowWaterRiverDetails::MakeInstance()
+{
+	return MakeShareable(new FShallowWaterRiverDetails);
+}
+
+FShallowWaterRiverDetails::~FShallowWaterRiverDetails()
+{
+	/*
+	if (GEngine)
+	{
+		GEngine->OnWorldDestroyed().RemoveAll(this);
+	}
+
+	FGameDelegates::Get().GetEndPlayMapDelegate().RemoveAll(this);
+	*/
+}
+
+void FShallowWaterRiverDetails::OnPiEEnd()
+{
+	UE_LOG(LogTemp, Log, TEXT("onPieEnd"));
+	if (Component.IsValid())
+	{
+		if (Component->GetOutermost()->HasAnyPackageFlags(PKG_PlayInEditor))
+		{
+			UE_LOG(LogTemp, Log, TEXT("onPieEnd - has package flags"));
+			UWorld* TheWorld = UWorld::FindWorldInPackage(Component->GetOutermost());
+			if (TheWorld)
+			{
+				OnWorldDestroyed(TheWorld);
+			}
+		}
+	}
+}
+
+void FShallowWaterRiverDetails::OnWorldDestroyed(class UWorld* InWorld)
+{
+	// We have to clear out any temp data interfaces that were bound to the component's package when the world goes away or otherwise
+	// we'll report GC leaks..
+	if (Component.IsValid())
+	{
+		if (Component->GetWorld() == InWorld)
+		{
+			UE_LOG(LogTemp, Log, TEXT("OnWorldDestroyed - matched up"));
+			Builder = nullptr;
+		}
+	}
+}
+
+void FShallowWaterRiverDetails::CustomizeDetails(IDetailLayoutBuilder& DetailBuilder)
+{
+	FPropertyEditorModule& PropertyModule = FModuleManager::GetModuleChecked<FPropertyEditorModule>("PropertyEditor");
+
+	Builder = &DetailBuilder;
+
+	static const FName ParamUtilitiesName = TEXT("ShallowWaterRiverComponent_Utilities");
+	IDetailCategoryBuilder& CustomCategory = DetailBuilder.EditCategory(ParamUtilitiesName, LOCTEXT("ParamUtilsCategoryName", "Utilities"), ECategoryPriority::Important);
+
+	CustomCategory.AddCustomRow(FText::GetEmpty())
+		.WholeRowContent()
+		.HAlign(HAlign_Left)
+		[
+			SNew(SBox)
+			.MaxDesiredWidth(300.f)
+			[
+				SNew(SUniformGridPanel)
+				.SlotPadding(2.0f)
+				+ SUniformGridPanel::Slot(0, 0)
+				[
+					SNew(SButton)
+					.OnClicked(this, &FShallowWaterRiverDetails::OnResetSelectedSystem)
+					.ToolTipText(LOCTEXT("ResetSystemButtonTooltip", "Resets the river system."))
+					.HAlign(HAlign_Center)
+					[
+						SNew(STextBlock)
+						.Text(LOCTEXT("ResetSystemButton", "Reset"))
+					]
+				]
+				+ SUniformGridPanel::Slot(1, 0)
+				[
+					SNew(SButton)
+					.OnClicked(this, &FShallowWaterRiverDetails::OnBakeSelectedSystem)
+					.ToolTipText(LOCTEXT("BakeSystemButtonTooltip", "Bakes the river system."))
+					.HAlign(HAlign_Center)
+					[
+						SNew(STextBlock)
+						.Text(LOCTEXT("BakeSystemButton", "Bake"))
+					]
+				]
+			]
+		];
+}
+
+FReply FShallowWaterRiverDetails::OnResetSelectedSystem()
+{
+	if (!Builder)
+		return FReply::Handled();
+
+	const TArray< TWeakObjectPtr<UObject> >& SelectedObjects = Builder->GetSelectedObjects();
+
+	for (int32 Idx = 0; Idx < SelectedObjects.Num(); ++Idx)
+	{
+		if (SelectedObjects[Idx].IsValid())
+		{
+			if (AActor* Actor = Cast<AActor>(SelectedObjects[Idx].Get()))
+			{
+				for (UActorComponent* AC : Actor->GetComponents())
+				{
+					UShallowWaterRiverComponent* ShallowWaterComponent = Cast<UShallowWaterRiverComponent>(AC);
+					if (ShallowWaterComponent)
+					{
+						ShallowWaterComponent->Rebuild();
+						ShallowWaterComponent->ReregisterComponent();
+					}
+				}
+			}
+			else if (UShallowWaterRiverComponent* ShallowWaterComponent = Cast<UShallowWaterRiverComponent>(SelectedObjects[Idx].Get()))
+			{
+				ShallowWaterComponent->Rebuild();
+				ShallowWaterComponent->ReregisterComponent();
+			}
+			
+		}
+	}
+
+	return FReply::Handled();
+}
+
+FReply FShallowWaterRiverDetails::OnBakeSelectedSystem()
+{
+	if (!Builder)
+		return FReply::Handled();
+
+	const TArray< TWeakObjectPtr<UObject> >& SelectedObjects = Builder->GetSelectedObjects();
+
+	for (int32 Idx = 0; Idx < SelectedObjects.Num(); ++Idx)
+	{
+		if (SelectedObjects[Idx].IsValid())
+		{
+			if (AActor* Actor = Cast<AActor>(SelectedObjects[Idx].Get()))
+			{
+				for (UActorComponent* AC : Actor->GetComponents())
+				{
+					UShallowWaterRiverComponent* ShallowWaterComponent = Cast<UShallowWaterRiverComponent>(AC);
+					if (ShallowWaterComponent)
+					{
+						ShallowWaterComponent->Bake();						
+					}
+				}
+			}
+			else if (UShallowWaterRiverComponent* ShallowWaterComponent = Cast<UShallowWaterRiverComponent>(SelectedObjects[Idx].Get()))
+			{
+				ShallowWaterComponent->Bake();
+			}
+
+		}
+	}
+
+	return FReply::Handled();
+}
+
+#undef LOCTEXT_NAMESPACE
+#endif
\ No newline at end of file
diff --git a/Engine/Plugins/Experimental/WaterAdvanced/Source/WaterAdvanced/Private/ShallowWaterRiverDetails.h b/Engine/Plugins/Experimental/WaterAdvanced/Source/WaterAdvanced/Private/ShallowWaterRiverDetails.h
new file mode 100644
--- /dev/null
+++ b/Engine/Plugins/Experimental/WaterAdvanced/Source/WaterAdvanced/Private/ShallowWaterRiverDetails.h
@@ -0,0 +1,36 @@
+// Copyright Epic Games, Inc. All Rights Reserved.
+
+#pragma once
+
+#if WITH_EDITOR
+#include "CoreMinimal.h"
+#include "IDetailCustomization.h"
+#include "Input/Reply.h"
+//#include "ShallowWaterRiverDetails.generated.h"
+
+class UShallowWaterRiverComponent;
+
+class FShallowWaterRiverDetails : public IDetailCustomization
+{
+public:
+	virtual ~FShallowWaterRiverDetails() override;
+
+	/** Makes a new instance of this detail layout class for a specific detail view requesting it */
+	static TSharedRef<IDetailCustomization> MakeInstance();
+
+	/** IDetailCustomization interface */
+	virtual void CustomizeDetails( IDetailLayoutBuilder& DetailBuilder ) override;
+
+protected:
+	void OnWorldDestroyed(class UWorld* InWorld);
+	void OnPiEEnd();
+
+	FReply OnResetSelectedSystem();
+	FReply OnBakeSelectedSystem();
+	
+private:
+	TWeakObjectPtr<UShallowWaterRiverComponent> Component;
+	IDetailLayoutBuilder* Builder = nullptr;
+
+};
+#endif
\ No newline at end of file
diff --git a/Engine/Plugins/Experimental/WaterAdvanced/Source/WaterAdvanced/Private/WaterAdvanced.cpp b/Engine/Plugins/Experimental/WaterAdvanced/Source/WaterAdvanced/Private/WaterAdvanced.cpp
--- a/Engine/Plugins/Experimental/WaterAdvanced/Source/WaterAdvanced/Private/WaterAdvanced.cpp
+++ b/Engine/Plugins/Experimental/WaterAdvanced/Source/WaterAdvanced/Private/WaterAdvanced.cpp
@@ -1,18 +1,31 @@
 // Copyright Epic Games, Inc. All Rights Reserved.
 
 #include "WaterAdvanced.h"
+#include "ShallowWaterRiverDetails.h"
+#include "ShallowWaterRiverActor.h"
+
+#if WITH_EDITOR
+#include "PropertyCustomizationHelpers.h"
+#endif
 
 #define LOCTEXT_NAMESPACE "FWaterAdvancedModule"
 
 void FWaterAdvancedModule::StartupModule()
 {
 	// This code will execute after your module is loaded into memory; the exact timing is specified in the .uplugin file per-module
+#if WITH_EDITOR
+	FPropertyEditorModule& PropertyModule = FModuleManager::LoadModuleChecked<FPropertyEditorModule>("PropertyEditor");
+
+	PropertyModule.RegisterCustomClassLayout(
+		UShallowWaterRiverComponent::StaticClass()->GetFName(),
+		FOnGetDetailCustomizationInstance::CreateStatic(&FShallowWaterRiverDetails::MakeInstance));
+#endif
 }
 
 void FWaterAdvancedModule::ShutdownModule()
 {
 	// This function may be called during shutdown to clean up your module.  For modules that support dynamic reloading,
-	// we call this function before unloading the module.
+	// we call this function before unloading the module.xture
 }
 
 #undef LOCTEXT_NAMESPACE
diff --git a/Engine/Plugins/Experimental/WaterAdvanced/Source/WaterAdvanced/Public/ShallowWaterRiverActor.h b/Engine/Plugins/Experimental/WaterAdvanced/Source/WaterAdvanced/Public/ShallowWaterRiverActor.h
--- a/Engine/Plugins/Experimental/WaterAdvanced/Source/WaterAdvanced/Public/ShallowWaterRiverActor.h
+++ b/Engine/Plugins/Experimental/WaterAdvanced/Source/WaterAdvanced/Public/ShallowWaterRiverActor.h
@@ -27,6 +27,12 @@ public:
 	UPROPERTY(EditAnywhere, Category = "Simulation", meta = (DisplayName = "Source Width"))
 	float SourceSize;
 
+	UPROPERTY(EditAnywhere, Category = "Simulation", meta = (DisplayName = "Speed"))
+	float SimSpeed = 1.f;
+
+	UPROPERTY(EditAnywhere, Category = "Simulation", meta = (DisplayName = "Num Substeps"))
+	int NumSteps = 1;
+
 	UPROPERTY(EditAnywhere, Category = "Water", meta = (DisplayName = "Source River Water Body"))
 	TObjectPtr<AWaterBody> SourceRiverWaterBody;
 
@@ -36,8 +42,11 @@ public:
 	UPROPERTY(EditAnywhere, Category = "Water", meta = (DisplayName = "Additional River Water Bodies"))
 	TArray<TObjectPtr<AWaterBody>> AdditonalRiverWaterBodies;
 
-	UPROPERTY(EditAnywhere, Category = "Reset", meta = (DisplayName = "Reset"))
-	bool Reset;
+	UPROPERTY(EditAnywhere, Category = "Baking", meta = (DisplayName = "Preview Baked Sim"))
+	bool PreviewBakedSim = false;
+
+	UPROPERTY(EditAnywhere, Category = "Shallow Water")
+	TObjectPtr<UTexture2D> BakedWaterSurfaceTexture;
 
 	virtual void PostLoad() override;
 
@@ -48,6 +57,8 @@ public:
 #if WITH_EDITOR
 	void Rebuild();
 
+	void Bake();
+
 	virtual void PostEditChangeProperty(FPropertyChangedEvent& PropertyChangedEvent) override;
 	void OnWaterInfoTextureCreated(const UTextureRenderTarget2D* InWaterInfoTexture);
 #endif // WITH_EDITOR
@@ -60,6 +71,9 @@ protected:
 	UPROPERTY(BlueprintReadOnly, VisibleDefaultsOnly, Category = "Shallow Water")
 	TObjectPtr<const UTextureRenderTarget2D> WaterInfoTexture;
 
+	UPROPERTY(BlueprintReadOnly, VisibleDefaultsOnly, Category = "Shallow Water")
+	TObjectPtr<UTextureRenderTarget2D> BakedWaterSurfaceRT;
+
 	bool QueryWaterAtSplinePoint(TObjectPtr<AWaterBody> WaterBody, int SplinePoint, FVector& OutPos, FVector& OutTangent, float& OutWidth, float& OutDepth);
 
 private:
diff --git a/Engine/Plugins/Experimental/WaterAdvanced/Source/WaterAdvanced/WaterAdvanced.Build.cs b/Engine/Plugins/Experimental/WaterAdvanced/Source/WaterAdvanced/WaterAdvanced.Build.cs
--- a/Engine/Plugins/Experimental/WaterAdvanced/Source/WaterAdvanced/WaterAdvanced.Build.cs
+++ b/Engine/Plugins/Experimental/WaterAdvanced/Source/WaterAdvanced/WaterAdvanced.Build.cs
@@ -43,8 +43,21 @@ public class WaterAdvanced : ModuleRules
 				"GameplayTags"
 			}
 		);
-		
-		
+
+		if (Target.bBuildEditor == true)
+		{
+			PrivateDependencyModuleNames.AddRange(
+				new string[] {
+				"TargetPlatform",
+				"DerivedDataCache",
+				"EditorFramework",
+				"UnrealEd",
+				"SlateCore",
+				"Slate"
+			});
+		}
+
+
 		DynamicallyLoadedModuleNames.AddRange(
 			new string[]
 			{
