disable sim when playing in split screen

diff --git a/Engine/Plugins/Experimental/WaterAdvanced/Source/WaterAdvanced/Private/ShallowWaterSubsystem.cpp b/Engine/Plugins/Experimental/WaterAdvanced/Source/WaterAdvanced/Private/ShallowWaterSubsystem.cpp
--- a/Engine/Plugins/Experimental/WaterAdvanced/Source/WaterAdvanced/Private/ShallowWaterSubsystem.cpp
+++ b/Engine/Plugins/Experimental/WaterAdvanced/Source/WaterAdvanced/Private/ShallowWaterSubsystem.cpp
@@ -107,7 +107,9 @@ void UShallowWaterSubsystem::PostInitialize()
 void UShallowWaterSubsystem::OnWorldBeginPlay(UWorld& InWorld)
 {
 	TArray<ULocalPlayer*> LocalPlayers = InWorld.GetGameInstance() != nullptr ? InWorld.GetGameInstance()->GetLocalPlayers() : TArray<ULocalPlayer*>();
-	if (!LocalPlayers.IsEmpty())
+	
+	// we don't support split screen
+	if (LocalPlayers.Num() == 1)
 	{
 		if (ULocalPlayer* const LocalPlayer = LocalPlayers[0])
 		{
@@ -119,6 +121,10 @@ void UShallowWaterSubsystem::OnWorldBeginPlay(UWorld& InWorld)
 			LocalPlayer->OnPlayerControllerChanged().AddUObject(this, &UShallowWaterSubsystem::OnLocalPlayerControllerBecomesValid);
 		}
 	}
+	else if (LocalPlayers.Num() > 1)
+	{
+		UE_LOG(LogShallowWater, Warning, TEXT("Shallow Water Simulation is disabled during splitscreen"));
+	}
 }
 
 bool UShallowWaterSubsystem::ShouldCreateSubsystem(UObject* Outer) const
@@ -231,6 +237,13 @@ void UShallowWaterSubsystem::InitializeShallowWater()
 		return;
 	}
 
+	// another check to make sure we don't simulate when we have split screen active
+	if (WeakPlayerController->GetSplitscreenPlayerCount() > 1)
+	{
+		UE_LOG(LogShallowWater, Warning, TEXT("Shallow Water Simulation is disabled during splitscreen"));
+		return;
+	}
+
 	if(!IsShallowWaterAllowedToInitialize())
 	{
 		return;
