scene captures for bottom contour

#changelist validated
#virtualized

diff --git a/Engine/Plugins/Experimental/WaterAdvanced/Source/WaterAdvanced/Private/ShallowWaterRiverActor.cpp b/Engine/Plugins/Experimental/WaterAdvanced/Source/WaterAdvanced/Private/ShallowWaterRiverActor.cpp
--- a/Engine/Plugins/Experimental/WaterAdvanced/Source/WaterAdvanced/Private/ShallowWaterRiverActor.cpp
+++ b/Engine/Plugins/Experimental/WaterAdvanced/Source/WaterAdvanced/Private/ShallowWaterRiverActor.cpp
@@ -7,9 +7,12 @@
 #include "WaterBodyActor.h"
 #include "WaterSplineComponent.h"
 #include "Engine/TextureRenderTarget2D.h"
+#include "TextureResource.h"
 #include "ShallowWaterCommon.h"
 #include "Engine/World.h"
 #include "Engine/Texture2D.h"
+#include "Math/Float16Color.h"
+#include "DataInterface/NiagaraDataInterfaceSceneCapture2D.h"
 
 UShallowWaterRiverComponent::UShallowWaterRiverComponent(const FObjectInitializer& Initializer)
 	: Super(Initializer)
@@ -51,6 +54,21 @@ void UShallowWaterRiverComponent::TickComponent(float DeltaTime, ELevelTick Tick
 	}
 
 #endif
+
+	// #todo(dmp): this is all temporary
+// set the sim texture on each
+	/*
+	for (TObjectPtr<AWaterBody > CurrWaterBody : AllWaterBodies)
+	{
+		TObjectPtr<UWaterBodyComponent> CurrWaterBodyComponent = CurrWaterBody->GetWaterBodyComponent();
+
+		if (CurrWaterBodyComponent->ShallowWaterSimulationGrid.ArrayValues.Num() == 0 && PreviewBakedSim)
+		{
+			CurrWaterBodyComponent->SetShallowWaterSimulationTexture(ShallowWaterSimArrayValues,
+				FIntVector2(BakedWaterSurfaceTexture->GetSizeX(), BakedWaterSurfaceTexture->GetSizeY()), SystemPos, WorldGridSize);
+		}		
+	}
+	*/
 }
 
 void UShallowWaterRiverComponent::OnUnregister()
@@ -95,6 +113,8 @@ void UShallowWaterRiverComponent::Rebuild()
 		RiverSimSystem->SetActive(false);
 		RiverSimSystem->DestroyComponent();
 		RiverSimSystem = nullptr;
+
+		AllWaterBodies.Empty();
 	}
 	
 	if (NiagaraRiverSimulation == nullptr)
@@ -102,9 +122,7 @@ void UShallowWaterRiverComponent::Rebuild()
 		UE_LOG(LogShallowWater, Warning, TEXT("UShallowWaterRiverComponent::Rebuild() - null Niagara system asset"));
 	}
 
-	// collect all the water bodies
-	TSet < TObjectPtr<AWaterBody>> AllWaterBodies;
-
+	// collect all the water bodies	
 	if (SourceRiverWaterBody != nullptr)
 	{
 		AllWaterBodies.Add(SourceRiverWaterBody);
@@ -188,36 +206,8 @@ void UShallowWaterRiverComponent::Rebuild()
 		UE_LOG(LogShallowWater, Warning, TEXT("UShallowWaterRiverComponent::Rebuild() - water sink query failed"));
 		return;
 	}
-	
 
-	/*
-				bool WaterInfoTextureSet = false;
-
-				if (!WaterInfoTextureSet)
-				{
-					AWaterZone* WaterZone = CurrWaterBodyComponent->GetWaterZone();
-					if (WaterZone)
-					{
-						const TObjectPtr<UTextureRenderTarget2D> NewWaterInfoTexture = WaterZone->WaterInfoTexture;
-
-						if (NewWaterInfoTexture == nullptr)
-						{
-							WaterZone->GetOnWaterInfoTextureCreated().RemoveDynamic(this, &UShallowWaterRiverComponent::OnWaterInfoTextureCreated);
-							WaterZone->GetOnWaterInfoTextureCreated().AddDynamic(this, &UShallowWaterRiverComponent::OnWaterInfoTextureCreated);
-						}
-						else
-						{
-							OnWaterInfoTextureCreated(NewWaterInfoTexture);
-						}
-
-						WaterInfoTextureSet = true;
-					}
-				}							
-	}
-	*/	
-
-
-	FVector SystemPos = CombinedBounds.Origin - FVector(0, 0, CombinedBounds.BoxExtent.Z);
+	SystemPos = CombinedBounds.Origin - FVector(0, 0, CombinedBounds.BoxExtent.Z);
 	
 	RiverSimSystem = NewObject<UNiagaraComponent>(this, NAME_None, RF_Transient);
 	RiverSimSystem->bUseAttachParentBound = false;
@@ -232,6 +222,20 @@ void UShallowWaterRiverComponent::Rebuild()
 
 		RiverSimSystem->SetVisibleFlag(true);
 		RiverSimSystem->SetAsset(NiagaraRiverSimulation);
+
+		UNiagaraDataInterfaceSceneCapture2D* BottomCaptureDI =
+			UNiagaraFunctionLibrary::GetDataInterface< UNiagaraDataInterfaceSceneCapture2D>(RiverSimSystem, "User.BottomCapture");
+
+		BottomCaptureDI->SourceMode = ENDISceneCapture2DSourceMode::Managed;
+		BottomCaptureDI->ManagedOrthoWidth = FMath::Max(WorldGridSize.X, WorldGridSize.Y);
+		BottomCaptureDI->ManagedTextureSize = FIntPoint(ResolutionMaxAxis, ResolutionMaxAxis);
+		BottomCaptureDI->ManagedCaptureSource = ESceneCaptureSource::SCS_SceneDepth;
+		BottomCaptureDI->ManagedTextureFormat = ETextureRenderTargetFormat::RTF_R16f;
+		BottomCaptureDI->ManagedProjectionType = ECameraProjectionMode::Orthographic;
+
+		BottomCaptureDI->ManagedShowOnlyActors.Empty();
+		BottomCaptureDI->ManagedShowOnlyActors.Append(BottomContourActors);
+
 		RiverSimSystem->ReinitializeSystem();
 	}
 	else
@@ -249,7 +253,8 @@ void UShallowWaterRiverComponent::Rebuild()
 
 	RiverSimSystem->Activate();
 
-	RiverSimSystem->SetVariableVec2(FName("WorldGridSize"), 2.0f * FVector2D(CombinedBounds.BoxExtent.X, CombinedBounds.BoxExtent.Y));
+	WorldGridSize = 2.0f * FVector2D(CombinedBounds.BoxExtent.X, CombinedBounds.BoxExtent.Y);
+	RiverSimSystem->SetVariableVec2(FName("WorldGridSize"), WorldGridSize);
 
 	// pad out source's box height a so it intersects the sim plane.  This value doesn't matter much so we hardcode it
 	float Overshoot = 1000.f;
@@ -288,6 +293,40 @@ void UShallowWaterRiverComponent::Bake()
 	BakedWaterSurfaceTexture = BakedWaterSurfaceRT->ConstructTexture2D(this, "BakedRiverTexture", TextureObjectFlags);
 
 	RiverSimSystem->SetVariableTexture(FName("BakedSimTexture"), BakedWaterSurfaceTexture);
+
+	// #todo(dmp): would love to do one more update here to run the pressure solver to get a pressure grid instead of velocity
+
+	// Readback to get the river texture values as an array
+	TArray<FFloat16Color> TmpShallowWaterSimArrayValues;
+	BakedWaterSurfaceRT->GameThread_GetRenderTargetResource()->ReadFloat16Pixels(TmpShallowWaterSimArrayValues);
+
+	ShallowWaterSimArrayValues.Empty();
+	ShallowWaterSimArrayValues.AddZeroed(TmpShallowWaterSimArrayValues.Num());
+
+	// #todo(dmp): no real need to cast all of the values to floats, but seems more conveinent for down the road compute 
+	int Index = 0;
+	for (FFloat16Color Val : TmpShallowWaterSimArrayValues)
+	{
+		FVector4 FloatVal;
+		FloatVal.X = Val.R;
+		FloatVal.Y = Val.G;
+		FloatVal.Z = Val.B;
+		FloatVal.W = Val.A;
+
+		ShallowWaterSimArrayValues[Index++] = FloatVal;
+	}
+
+	// #todo(dmp): this copies the array to each water body that uses it- it'd be better if it were referenced and owned by the
+	// river actor instead.  Perhaps the PT version copies it, but that is different...
+	// set the sim texture on each water body that is in the simulated river.  
+	for (TObjectPtr<AWaterBody > CurrWaterBody : AllWaterBodies)
+	{
+		TObjectPtr<UWaterBodyComponent> CurrWaterBodyComponent = CurrWaterBody->GetWaterBodyComponent();
+
+		// #todo(dmp): reenable
+		//CurrWaterBodyComponent->SetShallowWaterSimulationTexture(ShallowWaterSimArrayValues, 
+		//	FIntVector2(BakedWaterSurfaceRT->SizeX, BakedWaterSurfaceRT->SizeY), SystemPos, WorldGridSize);
+	}
 }
 
 void UShallowWaterRiverComponent::OnWaterInfoTextureCreated(const UTextureRenderTarget2D* InWaterInfoTexture)
diff --git a/Engine/Plugins/Experimental/WaterAdvanced/Source/WaterAdvanced/Public/ShallowWaterRiverActor.h b/Engine/Plugins/Experimental/WaterAdvanced/Source/WaterAdvanced/Public/ShallowWaterRiverActor.h
--- a/Engine/Plugins/Experimental/WaterAdvanced/Source/WaterAdvanced/Public/ShallowWaterRiverActor.h
+++ b/Engine/Plugins/Experimental/WaterAdvanced/Source/WaterAdvanced/Public/ShallowWaterRiverActor.h
@@ -3,8 +3,17 @@
 #pragma once
 
 #include "CoreMinimal.h"
+#include "Engine/EngineTypes.h"
 #include "GameFramework/Actor.h"
 #include "Components/PrimitiveComponent.h"
+
+#include "Math/Float16Color.h"
+#include "CoreMinimal.h"
+#include "Engine/EngineTypes.h"
+#include "NiagaraBakerOutput.h"
+#include "NiagaraBakerSettings.h"
+#include "NiagaraSystem.h"
+#include "UObject/GCObject.h"
 #include "ShallowWaterRiverActor.generated.h"
 
 class UNiagaraComponent;
@@ -48,6 +57,9 @@ public:
 	UPROPERTY(EditAnywhere, Category = "Shallow Water")
 	TObjectPtr<UTexture2D> BakedWaterSurfaceTexture;
 
+	UPROPERTY(EditAnywhere, Category = "Collisions")
+	TArray<TObjectPtr<AActor>> BottomContourActors;
+
 	virtual void PostLoad() override;
 
 	virtual void TickComponent(float DeltaTime, ELevelTick TickType, FActorComponentTickFunction* ThisTickFunction) override;
@@ -79,6 +91,18 @@ protected:
 private:
 	bool bIsInitialized;	
 	bool bTickInitialize;
+
+	UPROPERTY()
+	TSet < TObjectPtr<AWaterBody>> AllWaterBodies;
+
+	UPROPERTY()
+	FVector2D WorldGridSize;
+
+	UPROPERTY()
+	FVector SystemPos;
+
+	UPROPERTY()
+	TArray<FVector4> ShallowWaterSimArrayValues;
 };
 
 UCLASS(BlueprintType, HideCategories = (Physics, Replication, Input, Collision))
diff --git a/Engine/Plugins/Experimental/WaterAdvanced/Source/WaterAdvanced/WaterAdvanced.Build.cs b/Engine/Plugins/Experimental/WaterAdvanced/Source/WaterAdvanced/WaterAdvanced.Build.cs
--- a/Engine/Plugins/Experimental/WaterAdvanced/Source/WaterAdvanced/WaterAdvanced.Build.cs
+++ b/Engine/Plugins/Experimental/WaterAdvanced/Source/WaterAdvanced/WaterAdvanced.Build.cs
@@ -1,5 +1,6 @@
 // Copyright Epic Games, Inc. All Rights Reserved.
 
+using System.IO;
 using UnrealBuildTool;
 
 public class WaterAdvanced : ModuleRules
@@ -10,7 +11,7 @@ public class WaterAdvanced : ModuleRules
 		
 		PublicIncludePaths.AddRange(
 			new string[] {
-				// ... add public include paths required here ...
+				 System.IO.Path.Combine(GetModuleDirectory("Niagara"), "Private"),				 
 			}
 			);
 				
