remove private dependency and move implementation to function library

diff --git a/Engine/Plugins/Experimental/WaterAdvanced/Source/WaterAdvanced/Private/ShallowWaterRiverActor.cpp b/Engine/Plugins/Experimental/WaterAdvanced/Source/WaterAdvanced/Private/ShallowWaterRiverActor.cpp
--- a/Engine/Plugins/Experimental/WaterAdvanced/Source/WaterAdvanced/Private/ShallowWaterRiverActor.cpp
+++ b/Engine/Plugins/Experimental/WaterAdvanced/Source/WaterAdvanced/Private/ShallowWaterRiverActor.cpp
@@ -12,7 +12,6 @@
 #include "Engine/World.h"
 #include "Engine/Texture2D.h"
 #include "Math/Float16Color.h"
-#include "DataInterface/NiagaraDataInterfaceSceneCapture2D.h"
 
 UShallowWaterRiverComponent::UShallowWaterRiverComponent(const FObjectInitializer& Initializer)
 	: Super(Initializer)
@@ -96,7 +95,7 @@ void UShallowWaterRiverComponent::PostEditChangeProperty(FPropertyChangedEvent&
 	}
 			
 	// this should go before rebuild not after...something is wrong
-	if (PropertyName == GET_MEMBER_NAME_CHECKED(UShallowWaterRiverComponent, PreviewBakedSim))
+	if (PropertyName == GET_MEMBER_NAME_CHECKED(UShallowWaterRiverComponent, PreviewBakedSim) && RiverSimSystem->IsActive())
 	{
 		RiverSimSystem->SetVariableBool(FName("ReadCachedSim"), PreviewBakedSim);
 	}
@@ -112,9 +111,7 @@ void UShallowWaterRiverComponent::Rebuild()
 	{
 		RiverSimSystem->SetActive(false);
 		RiverSimSystem->DestroyComponent();
-		RiverSimSystem = nullptr;
-
-		AllWaterBodies.Empty();
+		RiverSimSystem = nullptr;		
 	}
 	
 	if (NiagaraRiverSimulation == nullptr)
@@ -122,6 +119,8 @@ void UShallowWaterRiverComponent::Rebuild()
 		UE_LOG(LogShallowWater, Warning, TEXT("UShallowWaterRiverComponent::Rebuild() - null Niagara system asset"));
 	}
 
+	AllWaterBodies.Empty();
+
 	// collect all the water bodies	
 	if (SourceRiverWaterBody != nullptr)
 	{
@@ -163,8 +162,8 @@ void UShallowWaterRiverComponent::Rebuild()
 		}
 		else
 		{
-			UE_LOG(LogShallowWater, Warning, TEXT("UShallowWaterRiverComponent::Rebuild() - null water body actor found"));
-			return;
+			UE_LOG(LogShallowWater, Warning, TEXT("UShallowWaterRiverComponent::Rebuild() - skipping null water body actor found"));
+			continue;
 		}
 	}
 	FBoxSphereBounds CombinedBounds(CombinedWorldBoundsBuilder);
@@ -222,21 +221,52 @@ void UShallowWaterRiverComponent::Rebuild()
 
 		RiverSimSystem->SetVisibleFlag(true);
 		RiverSimSystem->SetAsset(NiagaraRiverSimulation);
+							
+		// convert to raw ptr array for function library
+		TArray<AActor*> BottomContourActorsRawPtr;
+		for (TObjectPtr<AActor> CurrActor : BottomContourActors)
+		{
+			AActor* CurrActorRawPtr = CurrActor.Get();
+			BottomContourActorsRawPtr.Add(CurrActorRawPtr);
+		}
 
-		UNiagaraDataInterfaceSceneCapture2D* BottomCaptureDI =
-			UNiagaraFunctionLibrary::GetDataInterface< UNiagaraDataInterfaceSceneCapture2D>(RiverSimSystem, "User.BottomCapture");
+		FName DIName = "User.BottomCapture";
 
-		BottomCaptureDI->SourceMode = ENDISceneCapture2DSourceMode::Managed;
-		BottomCaptureDI->ManagedOrthoWidth = FMath::Max(WorldGridSize.X, WorldGridSize.Y);
-		BottomCaptureDI->ManagedTextureSize = FIntPoint(ResolutionMaxAxis, ResolutionMaxAxis);
-		BottomCaptureDI->ManagedCaptureSource = ESceneCaptureSource::SCS_SceneDepth;
-		BottomCaptureDI->ManagedTextureFormat = ETextureRenderTargetFormat::RTF_R16f;
-		BottomCaptureDI->ManagedProjectionType = ECameraProjectionMode::Orthographic;
+		UNiagaraFunctionLibrary::SetSceneCapture2DDataInterfaceManagedMode(RiverSimSystem, DIName,
+			ESceneCaptureSource::SCS_SceneDepth,
+			FIntPoint(ResolutionMaxAxis, ResolutionMaxAxis),
+			ETextureRenderTargetFormat::RTF_R16f,
+			ECameraProjectionMode::Orthographic,
+			90.0f,
+			FMath::Max(WorldGridSize.X, WorldGridSize.Y),
+			true,
+			false,
+			BottomContourActorsRawPtr);
 
-		BottomCaptureDI->ManagedShowOnlyActors.Empty();
-		BottomCaptureDI->ManagedShowOnlyActors.Append(BottomContourActors);
+	
+		// accumulate bounding box for river water bodies
+		FBoxSphereBounds::Builder BottomContourCombinedWorldBoundsBuilder;
+		for (TObjectPtr<AActor> BottomContourActor : BottomContourActors)
+		{
+			if (BottomContourActor != nullptr)
+			{
+				// accumulate bounds
+				FBoxSphereBounds WorldBounds;
+				BottomContourActor->GetActorBounds(false, WorldBounds.Origin, WorldBounds.BoxExtent);
 
+				BottomContourCombinedWorldBoundsBuilder += WorldBounds;
+			}
+			else
+			{
+				UE_LOG(LogShallowWater, Warning, TEXT("UShallowWaterRiverComponent::Rebuild() - skipping null bottom contour boundary actor found"));
+				continue;
+			}
+		}
+		FBoxSphereBounds CombinedBottomContourBounds(BottomContourCombinedWorldBoundsBuilder);
+				
 		RiverSimSystem->ReinitializeSystem();
+
+		RiverSimSystem->SetVariableFloat(FName("CaptureOffset"), BottomContourCaptureOffset + CombinedBottomContourBounds.Origin.Z + CombinedBottomContourBounds.BoxExtent.Z);
 	}
 	else
 	{
@@ -255,6 +285,7 @@ void UShallowWaterRiverComponent::Rebuild()
 
 	WorldGridSize = 2.0f * FVector2D(CombinedBounds.BoxExtent.X, CombinedBounds.BoxExtent.Y);
 	RiverSimSystem->SetVariableVec2(FName("WorldGridSize"), WorldGridSize);
+	RiverSimSystem->SetVariableInt(FName("ResolutionMaxAxis"), ResolutionMaxAxis);
 
 	// pad out source's box height a so it intersects the sim plane.  This value doesn't matter much so we hardcode it
 	float Overshoot = 1000.f;
@@ -324,7 +355,7 @@ void UShallowWaterRiverComponent::Bake()
 		TObjectPtr<UWaterBodyComponent> CurrWaterBodyComponent = CurrWaterBody->GetWaterBodyComponent();
 
 		// #todo(dmp): reenable
-		//CurrWaterBodyComponent->SetShallowWaterSimulationTexture(ShallowWaterSimArrayValues, 
+		// CurrWaterBodyComponent->SetShallowWaterSimulationTexture(ShallowWaterSimArrayValues, 
 		//	FIntVector2(BakedWaterSurfaceRT->SizeX, BakedWaterSurfaceRT->SizeY), SystemPos, WorldGridSize);
 	}
 }
diff --git a/Engine/Plugins/Experimental/WaterAdvanced/Source/WaterAdvanced/Public/ShallowWaterRiverActor.h b/Engine/Plugins/Experimental/WaterAdvanced/Source/WaterAdvanced/Public/ShallowWaterRiverActor.h
--- a/Engine/Plugins/Experimental/WaterAdvanced/Source/WaterAdvanced/Public/ShallowWaterRiverActor.h
+++ b/Engine/Plugins/Experimental/WaterAdvanced/Source/WaterAdvanced/Public/ShallowWaterRiverActor.h
@@ -60,6 +60,9 @@ public:
 	UPROPERTY(EditAnywhere, Category = "Collisions")
 	TArray<TObjectPtr<AActor>> BottomContourActors;
 
+	UPROPERTY(EditAnywhere, Category = "Collisions")
+	float BottomContourCaptureOffset = 0.f;
+
 	virtual void PostLoad() override;
 
 	virtual void TickComponent(float DeltaTime, ELevelTick TickType, FActorComponentTickFunction* ThisTickFunction) override;
diff --git a/Engine/Plugins/Experimental/WaterAdvanced/Source/WaterAdvanced/WaterAdvanced.Build.cs b/Engine/Plugins/Experimental/WaterAdvanced/Source/WaterAdvanced/WaterAdvanced.Build.cs
--- a/Engine/Plugins/Experimental/WaterAdvanced/Source/WaterAdvanced/WaterAdvanced.Build.cs
+++ b/Engine/Plugins/Experimental/WaterAdvanced/Source/WaterAdvanced/WaterAdvanced.Build.cs
@@ -9,13 +9,6 @@ public class WaterAdvanced : ModuleRules
 	{
 		PCHUsage = ModuleRules.PCHUsageMode.UseExplicitOrSharedPCHs;
 		
-		PublicIncludePaths.AddRange(
-			new string[] {
-				 System.IO.Path.Combine(GetModuleDirectory("Niagara"), "Private"),				 
-			}
-			);
-				
-		
 		PrivateIncludePaths.AddRange(
 			new string[] {
 				// ... add other private include paths required here ...
diff --git a/Engine/Plugins/FX/Niagara/Source/Niagara/Private/NiagaraFunctionLibrary.cpp b/Engine/Plugins/FX/Niagara/Source/Niagara/Private/NiagaraFunctionLibrary.cpp
--- a/Engine/Plugins/FX/Niagara/Source/Niagara/Private/NiagaraFunctionLibrary.cpp
+++ b/Engine/Plugins/FX/Niagara/Source/Niagara/Private/NiagaraFunctionLibrary.cpp
@@ -16,6 +16,7 @@
 #include "NiagaraComponentPool.h"
 #include "NiagaraWorldManager.h"
 #include "DataInterface/NiagaraDataInterfaceStaticMesh.h"
+#include "DataInterface/NiagaraDataInterfaceSceneCapture2D.h"
 #include "NiagaraDataInterfaceSkeletalMesh.h"
 #include "NiagaraDataInterfaceTexture.h"
 #include "NiagaraDataInterface2DArrayTexture.h"
@@ -510,6 +511,54 @@ void UNiagaraFunctionLibrary::SetSkeletalMeshDataInterfaceFilteredSockets(UNiaga
 	SkeletalMeshInterface->SetFilteredSocketsFromBlueprints(FilteredSockets);
 }
 
+NIAGARA_API void UNiagaraFunctionLibrary::SetSceneCapture2DDataInterfaceManagedMode(UNiagaraComponent* NiagaraSystem, const FName& DIName, 
+	ESceneCaptureSource ManagedCaptureSource, FIntPoint ManagedTextureSize, ETextureRenderTargetFormat ManagedTextureFormat,
+	ECameraProjectionMode::Type ManagedProjectionType, float ManagedFOVAngle, float ManagedOrthoWidth, bool bManagedCaptureEveryFrame,
+	bool bManagedCaptureOnMovement, const TArray<AActor*>& ShowOnlyActors)
+{
+	if (!NiagaraSystem)
+	{
+		if (FNiagaraUtilities::LogVerboseWarnings())
+		{
+			UE_LOG(LogNiagara, Warning, TEXT("NiagaraSystem in \"SetSceneCapture2DDataInterfaceManagedMode\" is NULL, skipping."));
+		}
+		return;
+	}
+
+	UNiagaraDataInterfaceSceneCapture2D* CaptureDI =
+		UNiagaraFunctionLibrary::GetDataInterface< UNiagaraDataInterfaceSceneCapture2D>(NiagaraSystem, DIName);
+
+	if (!CaptureDI)
+	{
+		if (FNiagaraUtilities::LogVerboseWarnings())
+		{
+			UE_LOG(LogNiagara, Warning, TEXT("Scene Capture 2D DI \"%s\" is Null in \"SetSceneCapture2DDataInterfaceManagedMode\" is NULL, skipping."), *DIName.ToString());
+		}
+		return;
+	}
+
+	CaptureDI->SourceMode = ENDISceneCapture2DSourceMode::Managed;
+
+	CaptureDI->ManagedCaptureSource = ManagedCaptureSource;
+	CaptureDI->ManagedTextureSize = ManagedTextureSize;
+	CaptureDI->ManagedTextureFormat = ManagedTextureFormat;
+	CaptureDI->ManagedProjectionType = ManagedProjectionType;
+	CaptureDI->ManagedFOVAngle = ManagedFOVAngle;
+	CaptureDI->ManagedOrthoWidth = ManagedOrthoWidth;
+	CaptureDI->bManagedCaptureEveryFrame = bManagedCaptureEveryFrame;
+	CaptureDI->bManagedCaptureOnMovement = bManagedCaptureOnMovement;
+	
+	CaptureDI->ManagedShowOnlyActors.Empty();
+
+	TArray<TObjectPtr<AActor>> AllActors;
+	for (AActor* CurrActor : ShowOnlyActors)
+	{
+		AllActors.Add(CurrActor);
+	}
+	
+	CaptureDI->ManagedShowOnlyActors.Append(AllActors);
+}
+
 void UNiagaraFunctionLibrary::SetTextureObject(UNiagaraComponent* NiagaraSystem, const FString& OverrideName, UTexture* Texture)
 {
 	if (!NiagaraSystem)
diff --git a/Engine/Plugins/FX/Niagara/Source/Niagara/Public/NiagaraFunctionLibrary.h b/Engine/Plugins/FX/Niagara/Source/Niagara/Public/NiagaraFunctionLibrary.h
--- a/Engine/Plugins/FX/Niagara/Source/Niagara/Public/NiagaraFunctionLibrary.h
+++ b/Engine/Plugins/FX/Niagara/Source/Niagara/Public/NiagaraFunctionLibrary.h
@@ -72,6 +72,19 @@ public:
 	UFUNCTION(BlueprintCallable, Category = Niagara)
 	static NIAGARA_API void SetSkeletalMeshDataInterfaceFilteredSockets(UNiagaraComponent* NiagaraSystem, const FString& OverrideName, const TArray<FName>& FilteredSockets);
 
+	/** Sets managed mode parameters for the Scene capture 2D  data interface, this is destructive as it modifies the data interface. */
+	UFUNCTION(BlueprintCallable, Category = Niagara)
+	static NIAGARA_API void SetSceneCapture2DDataInterfaceManagedMode(UNiagaraComponent* NiagaraSystem, const FName& DIName,
+		ESceneCaptureSource ManagedCaptureSource,
+		FIntPoint ManagedTextureSize,
+		ETextureRenderTargetFormat ManagedTextureFormat,
+		ECameraProjectionMode::Type ManagedProjectionType,
+		float ManagedFOVAngle,
+		float ManagedOrthoWidth,
+		bool bManagedCaptureEveryFrame,
+		bool bManagedCaptureOnMovement,
+		const TArray<AActor*> &ShowOnlyActors);
+
 	/** Overrides the Texture Object for a Niagara Texture Data Interface User Parameter.*/
 	UFUNCTION(BlueprintCallable, Category = Niagara)
 	static NIAGARA_API void SetTextureObject(UNiagaraComponent* NiagaraSystem, const FString& OverrideName, UTexture* Texture);
