runtime debug rendering switch

#virtualized

diff --git a/Engine/Plugins/Experimental/WaterAdvanced/Source/WaterAdvanced/Private/ShallowWaterSubsystem.cpp b/Engine/Plugins/Experimental/WaterAdvanced/Source/WaterAdvanced/Private/ShallowWaterSubsystem.cpp
--- a/Engine/Plugins/Experimental/WaterAdvanced/Source/WaterAdvanced/Private/ShallowWaterSubsystem.cpp
+++ b/Engine/Plugins/Experimental/WaterAdvanced/Source/WaterAdvanced/Private/ShallowWaterSubsystem.cpp
@@ -57,6 +57,14 @@ static FAutoConsoleVariableRef CVarSWDrawWaterSurfaceProjection(
 	TEXT("")
 );
 
+int32 GSWDebugRender = 0;
+static FAutoConsoleVariableRef CVarSWDebugRender(
+	TEXT("r.ShallowWater.DebugRender"),
+	GSWDebugRender,
+	TEXT("")
+);
+
+
 FName UShallowWaterSubsystem::ColliderComponentTag = FName("RigidMesh_ShallowWaterCollider");
 
 TArray<AWaterBody*> FShallowWaterCollisionTracker_Actor::GetOverlappingWaterBodies() const
@@ -514,6 +522,8 @@ void UShallowWaterSubsystem::UpdateGridMovement()
 	const float SecondsUntilDestroyed = FMath::Max(0, GSWRemainActiveForSeconds - SecondsSinceCollision);
 	ShallowWaterNiagaraSimulation->SetVariableFloat(FName("SecondsUntilDestroyed"), SecondsUntilDestroyed);
 
+	ShallowWaterNiagaraSimulation->SetVariableBool(FName("UseDebugRender"), GSWDebugRender == 1);
+
 #if ENABLE_DRAW_DEBUG
 	if (GSWDrawWaterSurfaceProjection)
 	{
@@ -1153,7 +1163,19 @@ void UShallowWaterSubsystem::OnWaterInfoTextureArrayCreated(const UTextureRender
 	WaterInfoTexture = InWaterInfoTexture;
 	if (ShallowWaterNiagaraSimulation)
 	{
-		ShallowWaterNiagaraSimulation->SetVariableTexture(FName("WaterInfoTexture"), Cast<UTexture>(const_cast<UTextureRenderTarget2DArray*>(WaterInfoTexture.Get())));
+		UTexture* WITTextureArray = Cast<UTexture>(const_cast<UTextureRenderTarget2DArray*>(WaterInfoTexture.Get()));
+		if (WITTextureArray == nullptr)
+		{
+			ensureMsgf(false, TEXT("UShallowWaterSubsystem::OnWaterInfoTextureCreated was called with Water Info Texture that isn't valid"));
+			return;
+		}
+
+		ShallowWaterNiagaraSimulation->SetVariableTexture(FName("WaterInfoTexture"), WITTextureArray);
+	}
+	else
+	{
+		ensureMsgf(false, TEXT("UShallowWaterSubsystem::OnWaterInfoTextureCreated was called with NULL ShallowWaterNiagaraSimulation"));
+		return;
 	}
 }
 
