more debugging options
async loading

#virtualized

diff --git a/Engine/Plugins/Experimental/WaterAdvanced/Source/WaterAdvanced/Private/ShallowWaterSubsystem.cpp b/Engine/Plugins/Experimental/WaterAdvanced/Source/WaterAdvanced/Private/ShallowWaterSubsystem.cpp
--- a/Engine/Plugins/Experimental/WaterAdvanced/Source/WaterAdvanced/Private/ShallowWaterSubsystem.cpp
+++ b/Engine/Plugins/Experimental/WaterAdvanced/Source/WaterAdvanced/Private/ShallowWaterSubsystem.cpp
@@ -64,6 +64,12 @@ static FAutoConsoleVariableRef CVarSWDebugRender(
 	TEXT("")
 );
 
+int32 GSWUseWaterInfoTexture = 1;
+static FAutoConsoleVariableRef CVarSWUseWaterInfoTexture(
+	TEXT("r.ShallowWater.UseWaterInfoTexture"),
+	GSWUseWaterInfoTexture,
+	TEXT("")
+);
 
 FName UShallowWaterSubsystem::ColliderComponentTag = FName("RigidMesh_ShallowWaterCollider");
 
@@ -111,6 +117,10 @@ void UShallowWaterSubsystem::PostInitialize()
 				})
 		);		
 	}
+	else
+	{
+		ensureMsgf(false, TEXT("UShallowWaterSubsystem::PostInitialize() - UShallowWaterSettings is not valid"));
+	}
 }
 
 void UShallowWaterSubsystem::OnWorldBeginPlay(UWorld& InWorld)
@@ -523,6 +533,7 @@ void UShallowWaterSubsystem::UpdateGridMovement()
 	ShallowWaterNiagaraSimulation->SetVariableFloat(FName("SecondsUntilDestroyed"), SecondsUntilDestroyed);
 
 	ShallowWaterNiagaraSimulation->SetVariableBool(FName("UseDebugRender"), GSWDebugRender == 1);
+	ShallowWaterNiagaraSimulation->SetVariableBool(FName("UseWaterInfoTexture"), GSWUseWaterInfoTexture == 1);
 
 #if ENABLE_DRAW_DEBUG
 	if (GSWDrawWaterSurfaceProjection)
@@ -645,9 +656,15 @@ void UShallowWaterSubsystem::RegisterPhysicsAssetProxiesDataAsset(UShallowWaterP
 {
 	if (Proxies == nullptr)
 	{
-		ensure(false);
+		UE_LOG(LogShallowWater, Warning, TEXT("ShallowWaterComponent: UShallowWaterPhysicsAssetOverridesDataAsset is NULL.  No vehicle interaction will be possible."))
+		return;
+	}
+	else if (Proxies->Overrides.Num() <= 0)
+	{
+		UE_LOG(LogShallowWater, Warning, TEXT("ShallowWaterComponent: Input UShallowWaterPhysicsAssetOverridesDataAsset: %s has 0 entries.  No additional vehicles will be supported."), *Proxies->GetName())
 		return;
 	}
+
 	for (const TTuple<FGameplayTag, FShallowWaterPhysicsAssetOverride>& Override : Proxies->Overrides)
 	{
 		if (RegisteredPhysicsAssetProxies.Contains(Override.Key))
@@ -657,6 +674,12 @@ void UShallowWaterSubsystem::RegisterPhysicsAssetProxiesDataAsset(UShallowWaterP
 		}
 	}
 	RegisteredPhysicsAssetProxies.Append(Proxies->Overrides);
+
+	if (RegisteredPhysicsAssetProxies.Num() <= 0)
+	{
+		UE_LOG(LogShallowWater, Warning, TEXT("ShallowWaterComponent: RegisteredPhysicsAssetProxies has 0 entries.  No vehicle interaction will be possible."))
+		return;
+	}
 }
 
 TSet<AWaterBody*> UShallowWaterSubsystem::GetAllOverlappingWaterBodiesAndUpdateCollisionTrackers() 
@@ -1003,6 +1026,10 @@ void UShallowWaterSubsystem::EnableCollisionForContext(const FShallowWaterCollis
 						})
 				);
 			}
+			else
+			{
+				UE_LOG(LogShallowWater, Warning, TEXT("EnableCollisionForContext() - Vehicle will not have collisions becvause no physics asset override was found"));
+			}
 		}
 		break;
 		}
