revert to using full physics asset for vehicles in shallow water

diff --git a/Engine/Plugins/Experimental/WaterAdvanced/Source/WaterAdvanced/Private/ShallowWaterSubsystem.cpp b/Engine/Plugins/Experimental/WaterAdvanced/Source/WaterAdvanced/Private/ShallowWaterSubsystem.cpp
--- a/Engine/Plugins/Experimental/WaterAdvanced/Source/WaterAdvanced/Private/ShallowWaterSubsystem.cpp
+++ b/Engine/Plugins/Experimental/WaterAdvanced/Source/WaterAdvanced/Private/ShallowWaterSubsystem.cpp
@@ -71,6 +71,13 @@ static FAutoConsoleVariableRef CVarSWUseWaterInfoTexture(
 	TEXT("")
 );
 
+int32 GSWUseFullVehiclePhysicsAssets = 1;
+static FAutoConsoleVariableRef CVarSWUUseFullVehiclePhysicsAssetse(
+	TEXT("r.ShallowWater.UseFullVehiclePhysicsAssets"),
+	GSWUseFullVehiclePhysicsAssets,
+	TEXT("")
+);
+
 FName UShallowWaterSubsystem::ColliderComponentTag = FName("RigidMesh_ShallowWaterCollider");
 
 TArray<AWaterBody*> FShallowWaterCollisionTracker_Actor::GetOverlappingWaterBodies() const
@@ -986,6 +993,7 @@ void UShallowWaterSubsystem::EnableCollisionForContext(const FShallowWaterCollis
 		// Apply PhysicsAsset override if defined in the data asset
 		// #todo confirm if SetLeaderPoseComponent actually works if Mesh is set to empty #PLAY-29387
 		UPhysicsAsset* PhysicsAssetOverride = nullptr;
+		bool IsSet = false;
 		if (!RegisteredPhysicsAssetProxies.IsEmpty())
 		{
 			const FGameplayTagContainer VehicleTags = GetVehicleTags(Context);
@@ -1001,13 +1009,15 @@ void UShallowWaterSubsystem::EnableCollisionForContext(const FShallowWaterCollis
 			}
 			if (FoundOverride)
 			{
+				IsSet = true;
+
 				TSoftObjectPtr<UPhysicsAsset> TmpPhysicsAsset = FoundOverride->PhysicsAsset;
 
 				UAssetManager::GetStreamableManager().RequestAsyncLoad(TmpPhysicsAsset.ToSoftObjectPath(),
 					FStreamableDelegate::CreateWeakLambda(this, [this, TmpPhysicsAsset, ProxyComp, Context]()
 						{								
 							if (TmpPhysicsAsset.IsValid())
-							{
+							{								
 								UPhysicsAsset* PhysicsAssetOverride = TmpPhysicsAsset.Get();
 
 								ProxyComp->SetPhysicsAsset(PhysicsAssetOverride);
@@ -1025,10 +1035,17 @@ void UShallowWaterSubsystem::EnableCollisionForContext(const FShallowWaterCollis
 							}
 						})
 				);
+			}			
+		}
+		if (!IsSet)
+		{
+			if (GSWUseFullVehiclePhysicsAssets)
+			{
+				Context.Component->ComponentTags.AddUnique(ColliderComponentTag);
 			}
 			else
 			{
-				UE_LOG(LogShallowWater, Warning, TEXT("EnableCollisionForContext() - Vehicle will not have collisions becvause no physics asset override was found"));
+			 	UE_LOG(LogShallowWater, Warning, TEXT("EnableCollisionForContext() - Vehicle will not have collisions because no physics asset override was found"));
 			}
 		}
 		break;
