set water zone parameters from C++

#virtualized

diff --git a/Engine/Plugins/Experimental/WaterAdvanced/Source/WaterAdvanced/Private/ShallowWaterSubsystem.cpp b/Engine/Plugins/Experimental/WaterAdvanced/Source/WaterAdvanced/Private/ShallowWaterSubsystem.cpp
--- a/Engine/Plugins/Experimental/WaterAdvanced/Source/WaterAdvanced/Private/ShallowWaterSubsystem.cpp
+++ b/Engine/Plugins/Experimental/WaterAdvanced/Source/WaterAdvanced/Private/ShallowWaterSubsystem.cpp
@@ -1190,6 +1190,22 @@ void UShallowWaterSubsystem::TryGetOrWaitForWaterInfoTextureFromWaterBodies(TSet
 	    			OnWaterInfoTextureArrayCreated(NewWaterInfoTexture);
 	    		}
 
+				if (ShallowWaterNiagaraSimulation)
+				{
+					const FVector2D ZoneLocation = FVector2D(WaterZone->GetDynamicWaterInfoCenter());
+
+					const FVector2D ZoneExtent = FVector2D(WaterZone->GetDynamicWaterInfoExtent());
+					const FVector2D WaterHeightExtents = FVector2D(WaterZone->GetWaterHeightExtents());
+					const float GroundZMin = WaterZone->GetGroundZMin();
+
+					ShallowWaterNiagaraSimulation->SetVariableVec2(FName("WaterZoneLocation"), ZoneLocation);
+					ShallowWaterNiagaraSimulation->SetVariableVec2(FName("WaterZoneExtent"), ZoneExtent);
+				}
+				else
+				{
+					ensureMsgf(false, TEXT("UShallowWaterSubsystem::TryGetOrWaitForWaterInfoTextureFromWaterBodies was called with NULL ShallowWaterNiagaraSimulation"));
+					return;
+				}
 	    		// Currently there can only be one unique WaterInfoTexture
 	    		break;
 	    	}
