pass waterzone index through

#virtualized

diff --git a/Engine/Plugins/Experimental/WaterAdvanced/Source/WaterAdvanced/Private/ShallowWaterSubsystem.cpp b/Engine/Plugins/Experimental/WaterAdvanced/Source/WaterAdvanced/Private/ShallowWaterSubsystem.cpp
--- a/Engine/Plugins/Experimental/WaterAdvanced/Source/WaterAdvanced/Private/ShallowWaterSubsystem.cpp
+++ b/Engine/Plugins/Experimental/WaterAdvanced/Source/WaterAdvanced/Private/ShallowWaterSubsystem.cpp
@@ -1200,6 +1200,7 @@ void UShallowWaterSubsystem::TryGetOrWaitForWaterInfoTextureFromWaterBodies(TSet
 
 					ShallowWaterNiagaraSimulation->SetVariableVec2(FName("WaterZoneLocation"), ZoneLocation);
 					ShallowWaterNiagaraSimulation->SetVariableVec2(FName("WaterZoneExtent"), ZoneExtent);
+					ShallowWaterNiagaraSimulation->SetVariableInt(FName("WaterZoneIdx"), WaterZone->GetWaterZoneIndex());
 				}
 				else
 				{
