enable base class subsystem for shallow water simulation
enable for QAGame

diff --git a/Engine/Plugins/Experimental/WaterAdvanced/Source/WaterAdvanced/Private/ShallowWaterSubsystem.cpp b/Engine/Plugins/Experimental/WaterAdvanced/Source/WaterAdvanced/Private/ShallowWaterSubsystem.cpp
--- a/Engine/Plugins/Experimental/WaterAdvanced/Source/WaterAdvanced/Private/ShallowWaterSubsystem.cpp
+++ b/Engine/Plugins/Experimental/WaterAdvanced/Source/WaterAdvanced/Private/ShallowWaterSubsystem.cpp
@@ -370,6 +370,17 @@ void UShallowWaterSubsystem::InitializeShallowWater()
 	UE_LOG(LogShallowWater, Log, TEXT("UShallowWaterSubsystem::InitializeShallowWater() finished successfully"));
 }
 
+bool UShallowWaterSubsystem::IsShallowWaterAllowedToInitialize() const
+{
+	TObjectPtr<UShallowWaterSettings> TmpSettings = GetMutableDefault<UShallowWaterSettings>();
+	if (TmpSettings)
+	{
+		return TmpSettings->UseDefaultShallowWaterSubsystem;
+	}
+
+	return false;
+}
+
 bool UShallowWaterSubsystem::IsShallowWaterInitialized() const
 {
 	return bIsShallowWaterInitialized;
diff --git a/Engine/Plugins/Experimental/WaterAdvanced/Source/WaterAdvanced/Public/ShallowWaterSettings.h b/Engine/Plugins/Experimental/WaterAdvanced/Source/WaterAdvanced/Public/ShallowWaterSettings.h
--- a/Engine/Plugins/Experimental/WaterAdvanced/Source/WaterAdvanced/Public/ShallowWaterSettings.h
+++ b/Engine/Plugins/Experimental/WaterAdvanced/Source/WaterAdvanced/Public/ShallowWaterSettings.h
@@ -36,6 +36,8 @@ public:
 	FName ResolutionMaxAxisMPCName = FName("FluidSimResolution");
 	UPROPERTY(Config, BlueprintReadOnly, EditAnywhere, Category="Shallow Water Simulation")
 	FName NormalRTMaterialName = FName("NormalAndHeight");
+	UPROPERTY(Config, BlueprintReadOnly, EditAnywhere, Category = "Shallow Water Simulation")
+	bool UseDefaultShallowWaterSubsystem = false;
 
 	UPROPERTY(Config, BlueprintReadOnly, EditAnywhere, Category="Interaction")
 	int32 MaxActivePawnNum = 6;
diff --git a/Engine/Plugins/Experimental/WaterAdvanced/Source/WaterAdvanced/Public/ShallowWaterSubsystem.h b/Engine/Plugins/Experimental/WaterAdvanced/Source/WaterAdvanced/Public/ShallowWaterSubsystem.h
--- a/Engine/Plugins/Experimental/WaterAdvanced/Source/WaterAdvanced/Public/ShallowWaterSubsystem.h
+++ b/Engine/Plugins/Experimental/WaterAdvanced/Source/WaterAdvanced/Public/ShallowWaterSubsystem.h
@@ -117,7 +117,7 @@ struct WATERADVANCED_API FShallowWaterCollisionTracker_Direct
 
 
 
-UCLASS(Abstract)
+UCLASS(BlueprintType, Transient)
 class WATERADVANCED_API UShallowWaterSubsystem : public UTickableWorldSubsystem
 {
 	GENERATED_BODY()
@@ -138,7 +138,7 @@ public:
 	 * by checking GameFeaturesSubsystem.GetPluginURLByName(PluginName, PluginURL))
 	 * and GameFeaturesSubsystem.IsGameFeaturePluginActive(PluginURL, true)
 	 */ 
-	virtual bool IsShallowWaterAllowedToInitialize() const { return true; }
+	virtual bool IsShallowWaterAllowedToInitialize() const;
 	bool IsShallowWaterInitialized() const;
 	APawn* GetNonSpectatorPawnFromWeakController() const;
 	TOptional<FVector> GetCameraLocationFromWeakController() const;
@@ -147,7 +147,7 @@ public:
 	// If the player doesn't have a physical pawn (e.g. spectating), the cursor jump around and lock on the nearest pawn to the camera
 	// If there is no relevant pawn (all eliminate), returns nullptr
 	APawn* GetTheMostRelevantPlayerPawn() const;
-	virtual FGameplayTagContainer GetVehicleTags(FShallowWaterCollisionContext Context) const PURE_VIRTUAL(UShallowWaterSubsystem::GetVehicleTags, return FGameplayTagContainer::EmptyContainer;)
+	virtual FGameplayTagContainer GetVehicleTags(FShallowWaterCollisionContext Context) const { return FGameplayTagContainer::EmptyContainer; }
 
 	void CreateRTs();
 	void InitializeParameters();
@@ -162,7 +162,7 @@ public:
 	void FlushPendingImpacts();
 	void WriteImpactToNDC(FVector ImpactPosition, FVector ImpactVelocity, float ImpactRadius);
 	// Override to return the ECC channel of bullets
-	virtual ECollisionChannel GetImpactCollisionChannel() PURE_VIRTUAL(UShallowWaterSubsystem::GetImpactCollisionChannel, return ECC_WorldDynamic;)
+	virtual ECollisionChannel GetImpactCollisionChannel() { return ECC_WorldDynamic; }
 	
 
 	// Manually set MID parameters for water bodies before they collide with any pawns or CollisionTrackers 
