make physics asset override a primary asset so it can be added to the cook with no references

#changelist validated
#virtualized

diff --git a/Engine/Plugins/Experimental/WaterAdvanced/Source/WaterAdvanced/Public/ShallowWaterPhysicsAssetOverridesDataAsset.h b/Engine/Plugins/Experimental/WaterAdvanced/Source/WaterAdvanced/Public/ShallowWaterPhysicsAssetOverridesDataAsset.h
--- a/Engine/Plugins/Experimental/WaterAdvanced/Source/WaterAdvanced/Public/ShallowWaterPhysicsAssetOverridesDataAsset.h
+++ b/Engine/Plugins/Experimental/WaterAdvanced/Source/WaterAdvanced/Public/ShallowWaterPhysicsAssetOverridesDataAsset.h
@@ -14,7 +14,7 @@ class USkeletalMesh;
  * 
  */
 UCLASS(Blueprintable)
-class WATERADVANCED_API UShallowWaterPhysicsAssetOverridesDataAsset : public UDataAsset
+class WATERADVANCED_API UShallowWaterPhysicsAssetOverridesDataAsset : public UPrimaryDataAsset
 {
 	GENERATED_BODY()
 
