GFA to ensure phyiscs overrides are loaded

diff --git a/Engine/Plugins/Experimental/WaterAdvanced/Source/WaterAdvanced/Private/ShallowWaterSubsystem.cpp b/Engine/Plugins/Experimental/WaterAdvanced/Source/WaterAdvanced/Private/ShallowWaterSubsystem.cpp
--- a/Engine/Plugins/Experimental/WaterAdvanced/Source/WaterAdvanced/Private/ShallowWaterSubsystem.cpp
+++ b/Engine/Plugins/Experimental/WaterAdvanced/Source/WaterAdvanced/Private/ShallowWaterSubsystem.cpp
@@ -102,33 +102,6 @@ UShallowWaterSubsystem::UShallowWaterSubsystem()
 {
 }
 
-void UShallowWaterSubsystem::PostInitialize()
-{
-	Super::PostInitialize();
-
-	// Register default PA Proxies before all other GFP chimes in
-	Settings = GetMutableDefault<UShallowWaterSettings>();
-	if (Settings)
-	{
-		UAssetManager::GetStreamableManager().RequestAsyncLoad(Settings->PhysicsAssetProxiesDataAsset.ToSoftObjectPath(),
-			FStreamableDelegate::CreateWeakLambda(this, [this]()
-				{
-					if (Settings->PhysicsAssetProxiesDataAsset.IsValid())
-					{
-						RegisterPhysicsAssetProxiesDataAsset(Settings->PhysicsAssetProxiesDataAsset.Get());						
-					}
-					else
-					{
-						ensureMsgf(false, TEXT("UShallowWaterSubsystem::PostInitialize() - PhyicsAssetProxiesDataAsset is not valid"));
-					}
-				})
-		);		
-	}
-	else
-	{
-		ensureMsgf(false, TEXT("UShallowWaterSubsystem::PostInitialize() - UShallowWaterSettings is not valid"));
-	}
-}
 
 void UShallowWaterSubsystem::OnWorldBeginPlay(UWorld& InWorld)
 {
@@ -286,6 +259,7 @@ void UShallowWaterSubsystem::InitializeShallowWater()
 		ObjectsToLoad.Add(Settings->WaterMPC.ToSoftObjectPath());
 		ObjectsToLoad.Add(Settings->DefaultShallowWaterNiagaraSimulation.ToSoftObjectPath());
 		ObjectsToLoad.Add(Settings->DefaultShallowWaterCollisionNDC.ToSoftObjectPath());
+		ObjectsToLoad.Add(Settings->PhysicsAssetProxiesDataAsset.ToSoftObjectPath());
 
 		UAssetManager::GetStreamableManager().RequestAsyncLoad(ObjectsToLoad,
 			FStreamableDelegate::CreateWeakLambda(this, [this]()
@@ -367,6 +341,16 @@ void UShallowWaterSubsystem::InitializeShallowWater()
 
 	PendingImpacts.Reset();
 	
+	// Register default PA Proxies before all other GFP chimes in
+	if (Settings->PhysicsAssetProxiesDataAsset.IsValid())
+	{
+		RegisterPhysicsAssetProxiesDataAsset(Settings->PhysicsAssetProxiesDataAsset.Get());
+	}
+	else
+	{
+		ensureMsgf(false, TEXT("UShallowWaterSubsystem::InitializeShallowWater() - PhyicsAssetProxiesDataAsset is not valid"));
+	}
+
 	UE_LOG(LogShallowWater, Log, TEXT("UShallowWaterSubsystem::InitializeShallowWater() finished successfully"));
 }
 
@@ -670,7 +654,7 @@ void UShallowWaterSubsystem::TryUpdateWaterBodyMIDParameters(UWaterBodyComponent
 	}
 }
 
-void UShallowWaterSubsystem::RegisterPhysicsAssetProxiesDataAsset(UShallowWaterPhysicsAssetOverridesDataAsset* Proxies)
+void UShallowWaterSubsystem::RegisterPhysicsAssetProxiesDataAsset(const UShallowWaterPhysicsAssetOverridesDataAsset* Proxies)
 {
 	if (Proxies == nullptr)
 	{
diff --git a/Engine/Plugins/Experimental/WaterAdvanced/Source/WaterAdvanced/Public/ShallowWaterSubsystem.h b/Engine/Plugins/Experimental/WaterAdvanced/Source/WaterAdvanced/Public/ShallowWaterSubsystem.h
--- a/Engine/Plugins/Experimental/WaterAdvanced/Source/WaterAdvanced/Public/ShallowWaterSubsystem.h
+++ b/Engine/Plugins/Experimental/WaterAdvanced/Source/WaterAdvanced/Public/ShallowWaterSubsystem.h
@@ -123,15 +123,13 @@ class WATERADVANCED_API UShallowWaterSubsystem : public UTickableWorldSubsystem
 	GENERATED_BODY()
 
 public:
-	UShallowWaterSubsystem();
-
-	virtual void PostInitialize() override;
+	UShallowWaterSubsystem();	
 	virtual void OnWorldBeginPlay(UWorld& InWorld) override;
 	virtual bool ShouldCreateSubsystem(UObject* Outer) const override;
 	virtual void Tick(float DeltaTime) override;
 	virtual TStatId GetStatId() const override;
 
-	void InitializeShallowWater();
+	virtual void InitializeShallowWater();
 	/*
 	 * Note: A world subsystem is always created and activated even when the plugin (GFP or not) it resides in is completely disabled
 	 * When the subclass of this class is a Game Feature Plugin. This can be used to limit when ShallowWater is actually enabled
@@ -176,7 +174,7 @@ public:
 	/*
 	 * Add PA overrides. Designed to be called by Game Feature Plugins. 
 	 */
-	void RegisterPhysicsAssetProxiesDataAsset(UShallowWaterPhysicsAssetOverridesDataAsset* Proxies);
+	void RegisterPhysicsAssetProxiesDataAsset(const UShallowWaterPhysicsAssetOverridesDataAsset* Proxies);
 
 protected:
 	// Asset can be set in Project Settings - Plugins - Water ShallowWaterSimulation
