fix initialization for subsystem
fix bullet impacts
fix missing skel mesh asset on proxy collider component for vehicles
make default subsystem abstract again to fix bullet impacts

diff --git a/Engine/Plugins/Experimental/WaterAdvanced/Source/WaterAdvanced/Private/ShallowWaterSubsystem.cpp b/Engine/Plugins/Experimental/WaterAdvanced/Source/WaterAdvanced/Private/ShallowWaterSubsystem.cpp
--- a/Engine/Plugins/Experimental/WaterAdvanced/Source/WaterAdvanced/Private/ShallowWaterSubsystem.cpp
+++ b/Engine/Plugins/Experimental/WaterAdvanced/Source/WaterAdvanced/Private/ShallowWaterSubsystem.cpp
@@ -102,6 +102,17 @@ UShallowWaterSubsystem::UShallowWaterSubsystem()
 {
 }
 
+void UShallowWaterSubsystem::PostInitialize()
+{
+	Super::PostInitialize();
+
+	// Register default PA Proxies before all other GFP chimes in
+	Settings = GetMutableDefault<UShallowWaterSettings>();
+	if (!Settings)
+	{
+		ensureMsgf(false, TEXT("UShallowWaterSubsystem::PostInitialize() - UShallowWaterSettings is not valid"));
+	}
+}
 
 void UShallowWaterSubsystem::OnWorldBeginPlay(UWorld& InWorld)
 {
@@ -348,7 +359,7 @@ void UShallowWaterSubsystem::InitializeShallowWater()
 	}
 	else
 	{
-		ensureMsgf(false, TEXT("UShallowWaterSubsystem::InitializeShallowWater() - PhyicsAssetProxiesDataAsset is not valid"));
+		UE_LOG(LogShallowWater, Log, TEXT("UShallowWaterSubsystem::InitializeShallowWater() - UShallowWaterSettings::PhyicsAssetProxiesDataAsset is not valid"));	
 	}
 
 	UE_LOG(LogShallowWater, Log, TEXT("UShallowWaterSubsystem::InitializeShallowWater() finished successfully"));
@@ -983,6 +994,7 @@ void UShallowWaterSubsystem::EnableCollisionForContext(const FShallowWaterCollis
 		const FString BaseName = TEXT("FluidsimCollisionProxy");
 		const FName CompName = MakeUniqueObjectName(this, USkeletalMeshComponent::StaticClass(), FName(*BaseName));
 		USkeletalMeshComponent* ProxyComp = NewObject<USkeletalMeshComponent>(Context.Component->GetOwner(), USkeletalMeshComponent::StaticClass(), CompName);
+		ProxyComp->SetSkeletalMeshAsset(Context.Component->GetSkeletalMeshAsset());
 
 		// Spawn an empty dummy SKM component to be used as Collision 
 		// Apply PhysicsAsset override if defined in the data asset
diff --git a/Engine/Plugins/Experimental/WaterAdvanced/Source/WaterAdvanced/Public/ShallowWaterSubsystem.h b/Engine/Plugins/Experimental/WaterAdvanced/Source/WaterAdvanced/Public/ShallowWaterSubsystem.h
--- a/Engine/Plugins/Experimental/WaterAdvanced/Source/WaterAdvanced/Public/ShallowWaterSubsystem.h
+++ b/Engine/Plugins/Experimental/WaterAdvanced/Source/WaterAdvanced/Public/ShallowWaterSubsystem.h
@@ -117,13 +117,14 @@ struct WATERADVANCED_API FShallowWaterCollisionTracker_Direct
 
 
 
-UCLASS(BlueprintType, Transient)
+UCLASS(Abstract)
 class WATERADVANCED_API UShallowWaterSubsystem : public UTickableWorldSubsystem
 {
 	GENERATED_BODY()
 
 public:
 	UShallowWaterSubsystem();	
+	virtual void PostInitialize() override;
 	virtual void OnWorldBeginPlay(UWorld& InWorld) override;
 	virtual bool ShouldCreateSubsystem(UObject* Outer) const override;
 	virtual void Tick(float DeltaTime) override;
