Changes to river simulation and new river template systems for various static scenarios (swirling, constant flow, zero velocity)
Modify ShallowWaterRiverActor to better support new integration with Water plugin
Various fixes for ShallowWaterRiverActor (ie: better landscape collision integration)

#virtualized

diff --git a/Engine/Plugins/Experimental/WaterAdvanced/Source/WaterAdvanced/Private/ShallowWaterRiverActor.cpp b/Engine/Plugins/Experimental/WaterAdvanced/Source/WaterAdvanced/Private/ShallowWaterRiverActor.cpp
--- a/Engine/Plugins/Experimental/WaterAdvanced/Source/WaterAdvanced/Private/ShallowWaterRiverActor.cpp
+++ b/Engine/Plugins/Experimental/WaterAdvanced/Source/WaterAdvanced/Private/ShallowWaterRiverActor.cpp
@@ -6,12 +6,15 @@
 #include "NiagaraFunctionLibrary.h"
 #include "WaterBodyActor.h"
 #include "WaterSplineComponent.h"
+#include "BakedShallowWaterSimulationComponent.h"
 #include "Engine/TextureRenderTarget2D.h"
 #include "TextureResource.h"
 #include "ShallowWaterCommon.h"
 #include "Engine/World.h"
 #include "Engine/Texture2D.h"
 #include "Math/Float16Color.h"
+#include "Kismet/GameplayStatics.h"
+#include "Landscape.h"
 
 UShallowWaterRiverComponent::UShallowWaterRiverComponent(const FObjectInitializer& Initializer)
 	: Super(Initializer)
@@ -27,59 +30,73 @@ UShallowWaterRiverComponent::UShallowWaterRiverComponent(const FObjectInitialize
 
 	ResolutionMaxAxis = 512;
 	SourceSize = 1000;
-	//NiagaraRiverSimulation = LoadObject<UNiagaraSystem>(nullptr, TEXT("/WaterAdvanced/Niagara/Systems/Grid2D_SW_River.Grid2D_SW_River"));
+
+	// #todo(dmp): default river system should be set here
+	//hap = LoadObject<UNiagaraSystem>(nullptr, TEXT("/WaterAdvanced/Niagara/Systems/Grid2D_SW_River.Grid2D_SW_River"));
+
+#if WITH_EDITOR
+	// Add all landscapes by default
+	/*
+	if (BottomContourActors.IsEmpty())
+	{
+		TArray<AActor*> FoundActors;
+		UGameplayStatics::GetAllActorsOfClass(GetWorld(), ALandscape::StaticClass(), FoundActors);
+
+		for (AActor* CurrActor : FoundActors)
+		{
+			BottomContourActors.Add(CurrActor);
+		}
+	}
+	*/
+
+	// start with one empty element
+	if (BottomContourActors.IsEmpty())
+	{
+		BottomContourActors.AddDefaulted();
+	}
+#endif
+
 }
 
 void UShallowWaterRiverComponent::PostLoad()
 {
 	Super::PostLoad();
 
-#if WITH_EDITOR
-	bIsInitialized = false;
-	bTickInitialize = false;
+	if (RenderState == EShallowWaterRenderState::LiveSim || RiverSimSystem == nullptr)
+	{
+	#if WITH_EDITOR
+		bIsInitialized = false;
+		bTickInitialize = false;
 
-	Rebuild();
-#endif
+		Rebuild();
+	#endif
+	}
+	else
+	{
+		RiverSimSystem->Activate();
+	}
 }
 
 void UShallowWaterRiverComponent::TickComponent(float DeltaTime, ELevelTick TickType, FActorComponentTickFunction* ThisTickFunction)
 {
 #if WITH_EDITOR
 	// lots of tick ordering issues, so we try to initialize on the first tick too
-	if (!bIsInitialized && !bTickInitialize)
+	if (RiverSimSystem == nullptr || (RenderState == EShallowWaterRenderState::LiveSim && !bIsInitialized && !bTickInitialize))
 	{
 		bTickInitialize = true;
 		Rebuild();
 	}
-
-#endif
-
-	// #todo(dmp): this is all temporary
-// set the sim texture on each
-	/*
-	for (TObjectPtr<AWaterBody > CurrWaterBody : AllWaterBodies)
+	else
 	{
-		TObjectPtr<UWaterBodyComponent> CurrWaterBodyComponent = CurrWaterBody->GetWaterBodyComponent();
-
-		if (CurrWaterBodyComponent->ShallowWaterSimulationGrid.ArrayValues.Num() == 0 && PreviewBakedSim)
-		{
-			CurrWaterBodyComponent->SetShallowWaterSimulationTexture(ShallowWaterSimArrayValues,
-				FIntVector2(BakedWaterSurfaceTexture->GetSizeX(), BakedWaterSurfaceTexture->GetSizeY()), SystemPos, WorldGridSize);
-		}		
+		RiverSimSystem->Activate();	
 	}
-	*/
+
+#endif
 }
 
 void UShallowWaterRiverComponent::OnUnregister()
 {
 	Super::OnUnregister();
-
-	//if (RiverSimSystem && !RiverSimSystem->IsBeingDestroyed())
-	//{
-	//	RiverSimSystem->SetActive(false);
-	//	RiverSimSystem->DestroyComponent();
-	//	RiverSimSystem = nullptr;
-	//}
 }
 
 #if WITH_EDITOR
@@ -95,14 +112,31 @@ void UShallowWaterRiverComponent::PostEditChangeProperty(FPropertyChangedEvent&
 	}
 			
 	// this should go before rebuild not after...something is wrong
-	if (PropertyName == GET_MEMBER_NAME_CHECKED(UShallowWaterRiverComponent, PreviewBakedSim) && RiverSimSystem->IsActive())
+	if (PropertyName == GET_MEMBER_NAME_CHECKED(UShallowWaterRiverComponent, RenderState) && RiverSimSystem != nullptr && RiverSimSystem->IsActive())
 	{
-		RiverSimSystem->SetVariableBool(FName("ReadCachedSim"), PreviewBakedSim);
+		RiverSimSystem->SetVariableBool(FName("ReadCachedSim"), RenderState == EShallowWaterRenderState::BakedSim);
+
+		// hide/show all the water bodies
 	}
 	else
 	{
 		Rebuild();
 	}
+
+	if (PropertyName == GET_MEMBER_NAME_CHECKED(UShallowWaterRiverComponent, RenderState))
+	{
+		RiverSimSystem->SetVisibility(RenderState != EShallowWaterRenderState::WaterComponent);
+
+		for (AWaterBody* CurrWaterBody : AllWaterBodies)
+		{
+			TObjectPtr<UWaterBodyComponent> CurrWaterBodyComponent = CurrWaterBody->GetWaterBodyComponent();
+
+			if (CurrWaterBodyComponent != nullptr)
+			{
+				CurrWaterBodyComponent->SetVisibility(RenderState == EShallowWaterRenderState::WaterComponent);
+			}
+		}
+	}
 }
 
 void UShallowWaterRiverComponent::Rebuild()
@@ -111,7 +145,7 @@ void UShallowWaterRiverComponent::Rebuild()
 	{
 		RiverSimSystem->SetActive(false);
 		RiverSimSystem->DestroyComponent();
-		RiverSimSystem = nullptr;		
+		RiverSimSystem = nullptr;
 	}
 	
 	if (NiagaraRiverSimulation == nullptr)
@@ -154,11 +188,16 @@ void UShallowWaterRiverComponent::Rebuild()
 	{
 		if (CurrWaterBody != nullptr)
 		{
-			// accumulate bounds
-			FBoxSphereBounds WorldBounds;
-			CurrWaterBody->GetActorBounds(false, WorldBounds.Origin, WorldBounds.BoxExtent);
+			TObjectPtr<UWaterBodyComponent> CurrWaterBodyComponent = CurrWaterBody->GetWaterBodyComponent();
+
+			if (CurrWaterBodyComponent != nullptr)
+			{
+				// accumulate bounds
+				FBoxSphereBounds WorldBounds;
+				CurrWaterBody->GetActorBounds(false, WorldBounds.Origin, WorldBounds.BoxExtent);
 
-			CombinedWorldBoundsBuilder += WorldBounds;
+				CombinedWorldBoundsBuilder += WorldBounds;
+			}
 		}
 		else
 		{
@@ -208,7 +247,7 @@ void UShallowWaterRiverComponent::Rebuild()
 
 	SystemPos = CombinedBounds.Origin - FVector(0, 0, CombinedBounds.BoxExtent.Z);
 	
-	RiverSimSystem = NewObject<UNiagaraComponent>(this, NAME_None, RF_Transient);
+	RiverSimSystem = NewObject<UNiagaraComponent>(this, NAME_None, RF_Public);
 	RiverSimSystem->bUseAttachParentBound = false;
 	RiverSimSystem->SetWorldLocation(SystemPos);
 
@@ -223,50 +262,58 @@ void UShallowWaterRiverComponent::Rebuild()
 		RiverSimSystem->SetAsset(NiagaraRiverSimulation);
 							
 		// convert to raw ptr array for function library
-		TArray<AActor*> BottomContourActorsRawPtr;
-		for (TObjectPtr<AActor> CurrActor : BottomContourActors)
+		if (bUseCapture)
 		{
-			AActor* CurrActorRawPtr = CurrActor.Get();
-			BottomContourActorsRawPtr.Add(CurrActorRawPtr);
-		}
+			TArray<AActor*> BottomContourActorsRawPtr;
+			for (TObjectPtr<AActor> CurrActor : BottomContourActors)
+			{
+				AActor* CurrActorRawPtr = CurrActor.Get();
+				BottomContourActorsRawPtr.Add(CurrActorRawPtr);
+			}
 
-		FName DIName = "User.BottomCapture";
+			FName DIName = "User.BottomCapture";
 
-		UNiagaraFunctionLibrary::SetSceneCapture2DDataInterfaceManagedMode(RiverSimSystem, DIName,
-			ESceneCaptureSource::SCS_SceneDepth,
-			FIntPoint(ResolutionMaxAxis, ResolutionMaxAxis),
-			ETextureRenderTargetFormat::RTF_R16f,
-			ECameraProjectionMode::Orthographic,
-			90.0f,
-			FMath::Max(WorldGridSize.X, WorldGridSize.Y),
-			true,
-			false,
-			BottomContourActorsRawPtr);
+			UNiagaraFunctionLibrary::SetSceneCapture2DDataInterfaceManagedMode(RiverSimSystem, DIName,
+				ESceneCaptureSource::SCS_SceneDepth,
+				FIntPoint(ResolutionMaxAxis, ResolutionMaxAxis),
+				ETextureRenderTargetFormat::RTF_R16f,
+				ECameraProjectionMode::Orthographic,
+				90.0f,
+				FMath::Max(WorldGridSize.X, WorldGridSize.Y),
+				true,
+				false,
+				BottomContourActorsRawPtr);
 
 	
-		// accumulate bounding box for river water bodies
-		FBoxSphereBounds::Builder BottomContourCombinedWorldBoundsBuilder;
-		for (TObjectPtr<AActor> BottomContourActor : BottomContourActors)
-		{
-			if (BottomContourActor != nullptr)
-			{
-				// accumulate bounds
-				FBoxSphereBounds WorldBounds;
-				BottomContourActor->GetActorBounds(false, WorldBounds.Origin, WorldBounds.BoxExtent);
-
-				BottomContourCombinedWorldBoundsBuilder += WorldBounds;
-			}
-			else
+			// accumulate bounding box for river water bodies
+			FBoxSphereBounds::Builder BottomContourCombinedWorldBoundsBuilder;
+			for (TObjectPtr<AActor> BottomContourActor : BottomContourActors)
 			{
-				UE_LOG(LogShallowWater, Warning, TEXT("UShallowWaterRiverComponent::Rebuild() - skipping null bottom contour boundary actor found"));
-				continue;
+				if (BottomContourActor != nullptr)
+				{
+					// accumulate bounds
+					FBoxSphereBounds WorldBounds;
+					BottomContourActor->GetActorBounds(false, WorldBounds.Origin, WorldBounds.BoxExtent);
+
+					BottomContourCombinedWorldBoundsBuilder += WorldBounds;
+				}
+				else
+				{
+					UE_LOG(LogShallowWater, Warning, TEXT("UShallowWaterRiverComponent::Rebuild() - skipping null bottom contour boundary actor found"));
+					continue;
+				}
 			}
-		}
-		FBoxSphereBounds CombinedBottomContourBounds(BottomContourCombinedWorldBoundsBuilder);
-				
-		RiverSimSystem->ReinitializeSystem();
+			FBoxSphereBounds CombinedBottomContourBounds(BottomContourCombinedWorldBoundsBuilder);				
 
-		RiverSimSystem->SetVariableFloat(FName("CaptureOffset"), BottomContourCaptureOffset + CombinedBottomContourBounds.Origin.Z + CombinedBottomContourBounds.BoxExtent.Z);
+
+			RiverSimSystem->ReinitializeSystem();
+	
+			RiverSimSystem->SetVariableFloat(FName("CaptureOffset"), BottomContourCaptureOffset + CombinedBottomContourBounds.Origin.Z + CombinedBottomContourBounds.BoxExtent.Z);
+		}
+		else
+		{
+			RiverSimSystem->ReinitializeSystem();
+		}
 	}
 	else
 	{
@@ -307,7 +354,7 @@ void UShallowWaterRiverComponent::Rebuild()
 	BakedWaterSurfaceRT = NewObject<UTextureRenderTarget2D>(this, NAME_None, RF_Transient);
 	BakedWaterSurfaceRT->InitAutoFormat(1, 1);
 	RiverSimSystem->SetVariableTextureRenderTarget(FName("SimGridRT"), BakedWaterSurfaceRT);
-	RiverSimSystem->SetVariableBool(FName("ReadCachedSim"), PreviewBakedSim);
+	RiverSimSystem->SetVariableBool(FName("ReadCachedSim"), RenderState == EShallowWaterRenderState::BakedSim);
 
 	if (BakedWaterSurfaceTexture != nullptr)
 	{
@@ -323,40 +370,42 @@ void UShallowWaterRiverComponent::Bake()
 
 	BakedWaterSurfaceTexture = BakedWaterSurfaceRT->ConstructTexture2D(this, "BakedRiverTexture", TextureObjectFlags);
 
-	RiverSimSystem->SetVariableTexture(FName("BakedSimTexture"), BakedWaterSurfaceTexture);
-
-	// #todo(dmp): would love to do one more update here to run the pressure solver to get a pressure grid instead of velocity
+	RiverSimSystem->SetVariableTexture(FName("BakedSimTexture"), BakedWaterSurfaceTexture);	
 
 	// Readback to get the river texture values as an array
 	TArray<FFloat16Color> TmpShallowWaterSimArrayValues;
 	BakedWaterSurfaceRT->GameThread_GetRenderTargetResource()->ReadFloat16Pixels(TmpShallowWaterSimArrayValues);
-
+		
+	TArray<FVector4> ShallowWaterSimArrayValues;
 	ShallowWaterSimArrayValues.Empty();
 	ShallowWaterSimArrayValues.AddZeroed(TmpShallowWaterSimArrayValues.Num());
 
-	// #todo(dmp): no real need to cast all of the values to floats, but seems more conveinent for down the road compute 
+	// cast all values to floats
 	int Index = 0;
 	for (FFloat16Color Val : TmpShallowWaterSimArrayValues)
 	{
+		const float WaterHeight = Val.R;
+		const float WaterDepth = Val.G;
+		const FVector2D WaterVelocity(Val.B, Val.A);
+		
 		FVector4 FloatVal;
-		FloatVal.X = Val.R;
-		FloatVal.Y = Val.G;
-		FloatVal.Z = Val.B;
-		FloatVal.W = Val.A;
+		FloatVal.X = WaterHeight;
+		FloatVal.Y = WaterDepth;
+		FloatVal.Z = WaterVelocity.X;
+		FloatVal.W = WaterVelocity.Y;
 
 		ShallowWaterSimArrayValues[Index++] = FloatVal;
 	}
 
-	// #todo(dmp): this copies the array to each water body that uses it- it'd be better if it were referenced and owned by the
-	// river actor instead.  Perhaps the PT version copies it, but that is different...
+	BakedSim = NewObject<UBakedShallowWaterSimulationComponent>(this, NAME_None, RF_Public);
+	BakedSim->SimulationData = FShallowWaterSimulationGrid(ShallowWaterSimArrayValues, FIntVector2(BakedWaterSurfaceRT->SizeX, BakedWaterSurfaceRT->SizeY), SystemPos, WorldGridSize);
+	
 	// set the sim texture on each water body that is in the simulated river.  
 	for (TObjectPtr<AWaterBody > CurrWaterBody : AllWaterBodies)
 	{
 		TObjectPtr<UWaterBodyComponent> CurrWaterBodyComponent = CurrWaterBody->GetWaterBodyComponent();
-
-		// #todo(dmp): reenable
-		// CurrWaterBodyComponent->SetShallowWaterSimulationTexture(ShallowWaterSimArrayValues, 
-		//	FIntVector2(BakedWaterSurfaceRT->SizeX, BakedWaterSurfaceRT->SizeY), SystemPos, WorldGridSize);
+				
+		CurrWaterBodyComponent->SetBakedShallowWaterSimulation(BakedSim);
 	}
 }
 
@@ -406,6 +455,9 @@ bool UShallowWaterRiverComponent::QueryWaterAtSplinePoint(TObjectPtr<AWaterBody>
 				OutDepth = Metadata->Depth.Points[SplinePoint].OutVal;
 
 				OutTangent = CurrSpline->SplineCurves.Position.Points[SplinePoint].LeaveTangent;
+
+				OutTangent = WaterBody->GetActorTransform().TransformVector(OutTangent);
+
 				OutTangent.Normalize();
 			}
 			else
diff --git a/Engine/Plugins/Experimental/WaterAdvanced/Source/WaterAdvanced/Public/ShallowWaterRiverActor.h b/Engine/Plugins/Experimental/WaterAdvanced/Source/WaterAdvanced/Public/ShallowWaterRiverActor.h
--- a/Engine/Plugins/Experimental/WaterAdvanced/Source/WaterAdvanced/Public/ShallowWaterRiverActor.h
+++ b/Engine/Plugins/Experimental/WaterAdvanced/Source/WaterAdvanced/Public/ShallowWaterRiverActor.h
@@ -6,7 +6,7 @@
 #include "Engine/EngineTypes.h"
 #include "GameFramework/Actor.h"
 #include "Components/PrimitiveComponent.h"
-
+#include "BakedShallowWaterSimulationComponent.h"
 #include "Math/Float16Color.h"
 #include "CoreMinimal.h"
 #include "Engine/EngineTypes.h"
@@ -21,6 +21,15 @@ class UNiagaraSystem;
 class AWaterBody;
 class UTextureRenderTarget2D;
 
+UENUM(BlueprintType)
+enum EShallowWaterRenderState : int
+{
+	WaterComponent,
+	LiveSim,
+	BakedSim
+};
+
+
 UCLASS(BlueprintType, HideCategories = (Physics, Replication, Input, Collision))
 class WATERADVANCED_API UShallowWaterRiverComponent : public UPrimitiveComponent
 {
@@ -50,18 +59,21 @@ public:
 
 	UPROPERTY(EditAnywhere, Category = "Water", meta = (DisplayName = "Additional River Water Bodies"))
 	TArray<TObjectPtr<AWaterBody>> AdditonalRiverWaterBodies;
-
-	UPROPERTY(EditAnywhere, Category = "Baking", meta = (DisplayName = "Preview Baked Sim"))
-	bool PreviewBakedSim = false;
+	
+	UPROPERTY(EditAnywhere, Category = "Rendering", meta = (DisplayName = "Render State"))
+	TEnumAsByte<EShallowWaterRenderState> RenderState = EShallowWaterRenderState::WaterComponent;
 
 	UPROPERTY(EditAnywhere, Category = "Shallow Water")
 	TObjectPtr<UTexture2D> BakedWaterSurfaceTexture;
+	
+	UPROPERTY(EditAnywhere, Category = "Collisions")
+	bool bUseCapture = false;
 
 	UPROPERTY(EditAnywhere, Category = "Collisions")
 	TArray<TObjectPtr<AActor>> BottomContourActors;
 
 	UPROPERTY(EditAnywhere, Category = "Collisions")
-	float BottomContourCaptureOffset = 0.f;
+	float BottomContourCaptureOffset = 1000.f;
 
 	virtual void PostLoad() override;
 
@@ -89,6 +101,9 @@ protected:
 	UPROPERTY(BlueprintReadOnly, VisibleDefaultsOnly, Category = "Shallow Water")
 	TObjectPtr<UTextureRenderTarget2D> BakedWaterSurfaceRT;
 
+	UPROPERTY(BlueprintReadOnly, VisibleDefaultsOnly, Category = "Shallow Water")
+	TObjectPtr<UBakedShallowWaterSimulationComponent> BakedSim;
+
 	bool QueryWaterAtSplinePoint(TObjectPtr<AWaterBody> WaterBody, int SplinePoint, FVector& OutPos, FVector& OutTangent, float& OutWidth, float& OutDepth);
 
 private:
@@ -103,9 +118,6 @@ private:
 
 	UPROPERTY()
 	FVector SystemPos;
-
-	UPROPERTY()
-	TArray<FVector4> ShallowWaterSimArrayValues;
 };
 
 UCLASS(BlueprintType, HideCategories = (Physics, Replication, Input, Collision))
diff --git a/Engine/Plugins/Experimental/WaterAdvanced/Source/WaterAdvanced/WaterAdvanced.Build.cs b/Engine/Plugins/Experimental/WaterAdvanced/Source/WaterAdvanced/WaterAdvanced.Build.cs
--- a/Engine/Plugins/Experimental/WaterAdvanced/Source/WaterAdvanced/WaterAdvanced.Build.cs
+++ b/Engine/Plugins/Experimental/WaterAdvanced/Source/WaterAdvanced/WaterAdvanced.Build.cs
@@ -34,7 +34,8 @@ public class WaterAdvanced : ModuleRules
 				"Water",
 				"DeveloperSettings",
 				"Projects",
-				"GameplayTags"
+				"GameplayTags",
+				"Landscape"
 			}
 		);
 
