first pass of water material integration
fix bounds computation on water body components
support for water info texture queries on river sim
fix height extrapolation for cells where no extrapolation sources are found
set keep water depth inside collider for river to avoid grid boundary artifacts
prototype code for computing 2D SDF for spline rivers

#changelist validated
#virtualized

diff --git a/Engine/Plugins/Experimental/WaterAdvanced/Source/WaterAdvanced/Private/ShallowWaterRiverActor.cpp b/Engine/Plugins/Experimental/WaterAdvanced/Source/WaterAdvanced/Private/ShallowWaterRiverActor.cpp
--- a/Engine/Plugins/Experimental/WaterAdvanced/Source/WaterAdvanced/Private/ShallowWaterRiverActor.cpp
+++ b/Engine/Plugins/Experimental/WaterAdvanced/Source/WaterAdvanced/Private/ShallowWaterRiverActor.cpp
@@ -8,6 +8,9 @@
 #include "WaterSplineComponent.h"
 #include "BakedShallowWaterSimulationComponent.h"
 #include "Engine/TextureRenderTarget2D.h"
+#include "Engine/TextureRenderTarget2DArray.h"
+#include "Materials/MaterialInstanceDynamic.h"
+
 #include "TextureResource.h"
 #include "ShallowWaterCommon.h"
 #include "Engine/World.h"
@@ -91,6 +94,28 @@ void UShallowWaterRiverComponent::TickComponent(float DeltaTime, ELevelTick Tick
 		RiverSimSystem->Activate();	
 	}
 
+	if (RenderState == EShallowWaterRenderState::WaterComponent)
+	{
+		for (AWaterBody* CurrWaterBody : AllWaterBodies)
+		{
+			TObjectPtr<UWaterBodyComponent> CurrWaterBodyComponent = CurrWaterBody->GetWaterBodyComponent();
+
+			if (CurrWaterBodyComponent != nullptr)
+			{
+				UMaterialInstanceDynamic* WaterMID = CurrWaterBodyComponent->GetWaterMaterialInstance();
+				// override materials on water bodies
+				if (WaterMID != nullptr)
+				{
+					WaterMID->SetTextureParameterValue("BakedWaterSimTex", BakedWaterSurfaceTexture);
+					WaterMID->SetVectorParameterValue("BakedWaterSimLocation", SystemPos);
+					WaterMID->SetVectorParameterValue("BakedWaterSimSize", FVector(WorldGridSize.X, WorldGridSize.Y, 1));
+
+					//FGuid BakedSimGuid;
+					//WaterMID->SetStaticSwitchParameterValueEditorOnly("UseBakedSim", true, BakedSimGuid);
+				}
+			}
+		}
+	}
 #endif
 }
 
@@ -125,17 +150,7 @@ void UShallowWaterRiverComponent::PostEditChangeProperty(FPropertyChangedEvent&
 
 	if (PropertyName == GET_MEMBER_NAME_CHECKED(UShallowWaterRiverComponent, RenderState))
 	{
-		RiverSimSystem->SetVisibility(RenderState != EShallowWaterRenderState::WaterComponent);
-
-		for (AWaterBody* CurrWaterBody : AllWaterBodies)
-		{
-			TObjectPtr<UWaterBodyComponent> CurrWaterBodyComponent = CurrWaterBody->GetWaterBodyComponent();
-
-			if (CurrWaterBodyComponent != nullptr)
-			{
-				CurrWaterBodyComponent->SetVisibility(RenderState == EShallowWaterRenderState::WaterComponent);
-			}
-		}
+		UpdateRenderState();
 	}
 }
 
@@ -194,10 +209,10 @@ void UShallowWaterRiverComponent::Rebuild()
 			{
 				// accumulate bounds
 				FBoxSphereBounds WorldBounds;
-				CurrWaterBody->GetActorBounds(false, WorldBounds.Origin, WorldBounds.BoxExtent);
+				CurrWaterBody->GetActorBounds(true, WorldBounds.Origin, WorldBounds.BoxExtent);				
 
 				CombinedWorldBoundsBuilder += WorldBounds;
-			}
+			}			
 		}
 		else
 		{
@@ -328,8 +343,102 @@ void UShallowWaterRiverComponent::Rebuild()
 		return;
 	}
 
-	RiverSimSystem->Activate();
+	// look for the water info texture
+	for (TObjectPtr<AWaterBody > CurrWaterBody : AllWaterBodies)
+	{
+		if (AWaterZone* WaterZone = CurrWaterBody->GetWaterBodyComponent()->GetWaterZone())
+		{
+			const TObjectPtr<UTextureRenderTarget2DArray> NewWaterInfoTexture = WaterZone->WaterInfoTextureArray;
+			if (NewWaterInfoTexture == nullptr)
+			{
+				WaterZone->GetOnWaterInfoTextureArrayCreated().RemoveDynamic(this, &UShallowWaterRiverComponent::OnWaterInfoTextureArrayCreated);
+				WaterZone->GetOnWaterInfoTextureArrayCreated().AddDynamic(this, &UShallowWaterRiverComponent::OnWaterInfoTextureArrayCreated);
+			}
+			else
+			{
+				OnWaterInfoTextureArrayCreated(NewWaterInfoTexture);
+			}
+
+/*
+			// generate high res texture signed distance to water body
+			FVector2D::FReal CellSize = FMath::Max(WorldGridSize.X, WorldGridSize.Y) / ResolutionMaxAxis;
+
+			FIntVector2 NumCells;
+			NumCells.X = int32(FMath::FloorToInt(WorldGridSize.X / CellSize));
+			NumCells.Y = int32(FMath::FloorToInt(WorldGridSize.Y / CellSize));
+
+			// Pad grid by 1 voxel if our computed bounding box is too small
+			if (WorldGridSize.X > WorldGridSize.Y && !FMath::IsNearlyEqual(CellSize * NumCells.Y, WorldGridSize.Y))
+			{
+				NumCells.Y++;
+			}
+			else if (WorldGridSize.X < WorldGridSize.Y && !FMath::IsNearlyEqual(CellSize * NumCells.X, WorldGridSize.X))
+			{
+				NumCells.X++;
+			}
+
+			UWaterSplineComponent* SplineComponent = CurrWaterBody->GetWaterBodyComponent()->GetWaterSpline();
+			UWaterSplineMetadata* SplineMetadata = CurrWaterBody->GetWaterBodyComponent()->GetWaterSplineMetadata();
+			FInterpCurveVector PositionCurve = SplineComponent->SplineCurves.Position;
+			const FInterpCurveFloat& WidthCurve = SplineMetadata->RiverWidth;
+			const Chaos::FRigidTransform3 WaterTransform = CurrWaterBody->GetWaterBodyComponent()->GetComponentTransform();
+
+			TArray<FInterpCurvePointVector> &PositionPts = PositionCurve.Points;
+
+			FInterpCurveVector PositionNew;
+			for (FInterpCurvePointVector Pt : PositionPts)
+			{
+				FVector NewPos = Pt.OutVal;
+				NewPos.Z = 0;
+
+				PositionNew.AddPoint(Pt.InVal, NewPos);
+			}			
+
+			TArray<FFloat16> SignedDistanceValues;
+			SignedDistanceValues.SetNum(NumCells.X * NumCells.Y);
+
+
+			for (int y = 0; y < NumCells.Y; ++y) {
+			for (int x = 0; x < NumCells.X; ++x) {
+				FVector2D CurrIndex = FVector2D(x, y);
+				
+
+				const FVector2D UnitPos = (CurrIndex + .5) / FVector2D(NumCells.X, NumCells.Y);
+				const FVector2D LocalPos = (UnitPos - .5) * WorldGridSize;
+				const FVector LocalPos3D = FVector(LocalPos.X, LocalPos.Y, 0);
+				FVector CellWorldPos = LocalPos3D + SystemPos;
+				FVector WaterLocalPos = WaterTransform.InverseTransformPosition(CellWorldPos);
+
+
+				// convert to world space
+				float Dist;
+				const float ClosestSplineKey = PositionCurve.FindNearest(WaterLocalPos, Dist);
+				FVector ClosestPos = PositionCurve.Eval(ClosestSplineKey);
+				float ClosestWidth = WidthCurve.Eval(ClosestSplineKey);
+
+				// snap to xy plane
+				//WaterLocalPos.Z = 0;
+				//ClosestPos.Z = 0;
+				
+				const float SignedDistanceToWater = (WaterLocalPos - ClosestPos).Length() - ClosestWidth;
+
+				SignedDistanceValues[x + NumCells.X * y] = SignedDistanceToWater < 0;
+			}}
 
+			SignedDistanceToSplineTexture = NewObject<UTexture2D>(GetTransientPackage(), NAME_None, RF_Public);
+			SignedDistanceToSplineTexture->Source.Init(NumCells.X, NumCells.Y, 1, 1, TSF_R16F, (const uint8*) SignedDistanceValues.GetData());
+			SignedDistanceToSplineTexture->PowerOfTwoMode = ETexturePowerOfTwoSetting::None;
+			SignedDistanceToSplineTexture->MipGenSettings = TextureMipGenSettings::TMGS_NoMipmaps;			
+			SignedDistanceToSplineTexture->UpdateResource();
+			SignedDistanceToSplineTexture->PostEditChange();
+*/
+
+			break;
+		}
+	}
+
+	RiverSimSystem->Activate();
+	
 	WorldGridSize = 2.0f * FVector2D(CombinedBounds.BoxExtent.X, CombinedBounds.BoxExtent.Y);
 	RiverSimSystem->SetVariableVec2(FName("WorldGridSize"), WorldGridSize);
 	RiverSimSystem->SetVariableInt(FName("ResolutionMaxAxis"), ResolutionMaxAxis);
@@ -409,25 +518,6 @@ void UShallowWaterRiverComponent::Bake()
 	}
 }
 
-void UShallowWaterRiverComponent::OnWaterInfoTextureCreated(const UTextureRenderTarget2D* InWaterInfoTexture)
-{
-	if (InWaterInfoTexture == nullptr)
-	{
-		ensureMsgf(false, TEXT("UShallowWaterRiverComponent::OnWaterInfoTextureCreated was called with NULL WaterInfoTexture"));
-		return;
-	}
-
-	WaterInfoTexture = InWaterInfoTexture;
-	if (RiverSimSystem)
-	{
-		RiverSimSystem->SetVariableTexture(FName("WaterInfoTexture"), Cast<UTexture>(const_cast<UTextureRenderTarget2D*>(WaterInfoTexture.Get())));
-	}
-	else
-	{
-		ensureMsgf(false, TEXT("UShallowWaterRiverComponent::OnWaterInfoTextureCreated No river system to set water info on"));
-	}
-}
-
 bool UShallowWaterRiverComponent::QueryWaterAtSplinePoint(TObjectPtr<AWaterBody> WaterBody, int SplinePoint, FVector& OutPos, FVector& OutTangent, float& OutWidth, float& OutDepth)
 {	
 	if (WaterBody != nullptr)
@@ -480,6 +570,57 @@ bool UShallowWaterRiverComponent::QueryWaterAtSplinePoint(TObjectPtr<AWaterBody>
 
 	return true;
 }
+
+void UShallowWaterRiverComponent::UpdateRenderState()
+{
+	bool RenderWaterBody = RenderState == EShallowWaterRenderState::WaterComponent;
+
+	RiverSimSystem->SetVisibility(!RenderWaterBody);
+
+	for (AWaterBody* CurrWaterBody : AllWaterBodies)
+	{
+		TObjectPtr<UWaterBodyComponent> CurrWaterBodyComponent = CurrWaterBody->GetWaterBodyComponent();
+
+		if (CurrWaterBodyComponent != nullptr)
+		{
+			CurrWaterBodyComponent->SetVisibility(RenderWaterBody);
+
+			UMaterialInstanceDynamic* WaterMID = CurrWaterBodyComponent->GetWaterMaterialInstance();
+			// override materials on water bodies							
+			WaterMID->SetTextureParameterValue("BakedWaterSimTex", BakedWaterSurfaceTexture);
+			WaterMID->SetVectorParameterValue("BakedWaterSimLocation", SystemPos);
+			WaterMID->SetVectorParameterValue("BakedWaterSimSize", FVector(WorldGridSize.X, WorldGridSize.Y, 1));			
+		}
+	}
+}
+
+void UShallowWaterRiverComponent::OnWaterInfoTextureArrayCreated(const UTextureRenderTarget2DArray* InWaterInfoTexture)
+{
+	if (InWaterInfoTexture == nullptr)
+	{
+		ensureMsgf(false, TEXT("UShallowWaterRiverComponent::OnWaterInfoTextureCreated was called with NULL WaterInfoTexture"));
+		return;
+	}
+
+	WaterInfoTexture = InWaterInfoTexture;
+	if (RiverSimSystem)
+	{
+		UTexture* WITTextureArray = Cast<UTexture>(const_cast<UTextureRenderTarget2DArray*>(WaterInfoTexture.Get()));
+		if (WITTextureArray == nullptr)
+		{
+			ensureMsgf(false, TEXT("UShallowWaterRiverComponent::OnWaterInfoTextureCreated was called with Water Info Texture that isn't valid"));
+			return;
+		}
+
+		RiverSimSystem->SetVariableTexture(FName("WaterInfoTexture"), WITTextureArray);
+	}
+	else
+	{
+		ensureMsgf(false, TEXT("UShallowWaterRiverComponent::OnWaterInfoTextureCreated was called with NULL ShallowWaterNiagaraSimulation"));
+		return;
+	}
+}
+
 #endif
 
 AShallowWaterRiver::AShallowWaterRiver(const FObjectInitializer& ObjectInitializer)
diff --git a/Engine/Plugins/Experimental/WaterAdvanced/Source/WaterAdvanced/Private/ShallowWaterRiverDetails.cpp b/Engine/Plugins/Experimental/WaterAdvanced/Source/WaterAdvanced/Private/ShallowWaterRiverDetails.cpp
--- a/Engine/Plugins/Experimental/WaterAdvanced/Source/WaterAdvanced/Private/ShallowWaterRiverDetails.cpp
+++ b/Engine/Plugins/Experimental/WaterAdvanced/Source/WaterAdvanced/Private/ShallowWaterRiverDetails.cpp
@@ -170,7 +170,8 @@ FReply FShallowWaterRiverDetails::OnBakeSelectedSystem()
 					UShallowWaterRiverComponent* ShallowWaterComponent = Cast<UShallowWaterRiverComponent>(AC);
 					if (ShallowWaterComponent)
 					{
-						ShallowWaterComponent->Bake();						
+						ShallowWaterComponent->Bake();
+						break;
 					}
 				}
 			}
diff --git a/Engine/Plugins/Experimental/WaterAdvanced/Source/WaterAdvanced/Public/ShallowWaterRiverActor.h b/Engine/Plugins/Experimental/WaterAdvanced/Source/WaterAdvanced/Public/ShallowWaterRiverActor.h
--- a/Engine/Plugins/Experimental/WaterAdvanced/Source/WaterAdvanced/Public/ShallowWaterRiverActor.h
+++ b/Engine/Plugins/Experimental/WaterAdvanced/Source/WaterAdvanced/Public/ShallowWaterRiverActor.h
@@ -14,12 +14,13 @@
 #include "NiagaraBakerSettings.h"
 #include "NiagaraSystem.h"
 #include "UObject/GCObject.h"
+#include "Engine/TextureRenderTarget2D.h"
+#include "Engine/TextureRenderTarget2DArray.h"
 #include "ShallowWaterRiverActor.generated.h"
 
 class UNiagaraComponent;
 class UNiagaraSystem;
 class AWaterBody;
-class UTextureRenderTarget2D;
 
 UENUM(BlueprintType)
 enum EShallowWaterRenderState : int
@@ -66,6 +67,9 @@ public:
 	UPROPERTY(EditAnywhere, Category = "Shallow Water")
 	TObjectPtr<UTexture2D> BakedWaterSurfaceTexture;
 	
+	//UPROPERTY(EditAnywhere, Category = "Shallow Water")
+	//TObjectPtr<UTexture2D> SignedDistanceToSplineTexture;
+
 	UPROPERTY(EditAnywhere, Category = "Collisions")
 	bool bUseCapture = false;
 
@@ -87,7 +91,7 @@ public:
 	void Bake();
 
 	virtual void PostEditChangeProperty(FPropertyChangedEvent& PropertyChangedEvent) override;
-	void OnWaterInfoTextureCreated(const UTextureRenderTarget2D* InWaterInfoTexture);
+	void OnWaterInfoTextureArrayCreated(const UTextureRenderTarget2DArray* InWaterInfoTexture);
 #endif // WITH_EDITOR
 
 protected:
@@ -96,7 +100,7 @@ protected:
 	TObjectPtr<UNiagaraComponent> RiverSimSystem;
 
 	UPROPERTY(BlueprintReadOnly, VisibleDefaultsOnly, Category = "Shallow Water")
-	TObjectPtr<const UTextureRenderTarget2D> WaterInfoTexture;
+	TObjectPtr<const UTextureRenderTarget2DArray> WaterInfoTexture;
 
 	UPROPERTY(BlueprintReadOnly, VisibleDefaultsOnly, Category = "Shallow Water")
 	TObjectPtr<UTextureRenderTarget2D> BakedWaterSurfaceRT;
@@ -106,6 +110,8 @@ protected:
 
 	bool QueryWaterAtSplinePoint(TObjectPtr<AWaterBody> WaterBody, int SplinePoint, FVector& OutPos, FVector& OutTangent, float& OutWidth, float& OutDepth);
 
+	void UpdateRenderState();
+
 private:
 	bool bIsInitialized;	
 	bool bTickInitialize;
