support cpu queries for shallow water simulation
more rendering/sim options exposed on river actor

diff --git a/Engine/Plugins/Experimental/Buoyancy/Source/Runtime/Private/BuoyancySubsystem.cpp b/Engine/Plugins/Experimental/Buoyancy/Source/Runtime/Private/BuoyancySubsystem.cpp
--- a/Engine/Plugins/Experimental/Buoyancy/Source/Runtime/Private/BuoyancySubsystem.cpp
+++ b/Engine/Plugins/Experimental/Buoyancy/Source/Runtime/Private/BuoyancySubsystem.cpp
@@ -371,7 +371,7 @@ void UBuoyancySubsystem::UpdateSplineData()
 					WaterBodyComponent->GetWaterBodyType(),
 					SplineMetadata ? SplineMetadata->RiverWidth : EmptyOptionalFloat,
 					SplineMetadata ? SplineMetadata->WaterVelocityScalar : EmptyOptionalFloat,
-					bUseShallowWaterSimulation && WaterBodyComponent->UseBakedShallowWaterSimulationForBuoyancy() ? BakedSim->SimulationData : EmptyOptionalSimGrid
+					bUseShallowWaterSimulation && WaterBodyComponent->UseBakedSimulationForQueriesAndPhysics() ? BakedSim->SimulationData : EmptyOptionalSimGrid
 				);
 
 				// Go over each physics object in each primitive component which was generated
diff --git a/Engine/Plugins/Experimental/Water/Source/Runtime/Private/WaterBodyComponent.cpp b/Engine/Plugins/Experimental/Water/Source/Runtime/Private/WaterBodyComponent.cpp
--- a/Engine/Plugins/Experimental/Water/Source/Runtime/Private/WaterBodyComponent.cpp
+++ b/Engine/Plugins/Experimental/Water/Source/Runtime/Private/WaterBodyComponent.cpp
@@ -624,6 +624,17 @@ FWaterBodyQueryResult UWaterBodyComponent::QueryWaterInfoClosestToWorldLocation(
 	const bool bFlatSurface = IsFlatSurface();
 	const UWaterSplineComponent* const WaterSpline = GetWaterSpline();
 
+	FVector ShallowWaterVelocity = { 0, 0, 0 };
+	float ShallowWaterHeight = 0;
+	float ShallowWaterDepth = 0;
+	const bool bUseShallowWaterValues = UseBakedSimulationForQueriesAndPhysics();
+
+	if (bUseShallowWaterValues)
+	{		
+		BakedShallowWaterSim->SimulationData.SampleShallowWaterSimulationAtPosition(InWorldLocation,
+			ShallowWaterVelocity, ShallowWaterHeight, ShallowWaterDepth);
+	}
+
 	// Compute water plane location :
 	if (EnumHasAnyFlags(Result.GetQueryFlags(), EWaterBodyQueryFlags::ComputeLocation))
 	{
@@ -633,7 +644,22 @@ FWaterBodyQueryResult UWaterBodyComponent::QueryWaterInfoClosestToWorldLocation(
 		//  If the user fails to do so, at least it allows immersion depth to be 0.0f, which means the query location is NOT in water :
 		if (!Result.IsInExclusionVolume())
 		{
-			WaterPlaneLocation.Z = (bFlatSurface || WaterSpline  == nullptr) ? GetComponentLocation().Z : WaterSpline->GetLocationAtSplineInputKey(Result.LazilyComputeSplineKey(*this, InWorldLocation), ESplineCoordinateSpace::World).Z;
+			float WaterPlaneHeight;
+
+			if (bFlatSurface || WaterSpline == nullptr)
+			{
+				WaterPlaneHeight = GetComponentLocation().Z;
+			}
+			else if (bUseShallowWaterValues)
+			{
+				WaterPlaneHeight = ShallowWaterHeight;
+			}
+			else
+			{
+				WaterPlaneHeight = WaterSpline->GetLocationAtSplineInputKey(Result.LazilyComputeSplineKey(*this, InWorldLocation), ESplineCoordinateSpace::World).Z;
+			}
+
+			WaterPlaneLocation.Z = WaterPlaneHeight;
 
 			// Apply body height offset if applicable (ocean)
 			if (IsHeightOffsetSupported())
@@ -655,8 +681,15 @@ FWaterBodyQueryResult UWaterBodyComponent::QueryWaterInfoClosestToWorldLocation(
 		// Default to Z up for the normal
 		if (!bFlatSurface && WaterSpline != nullptr)
 		{
-			// For rivers default to using spline up vector to account for sloping rivers
-			WaterPlaneNormal = WaterSpline->GetUpVectorAtSplineInputKey(Result.LazilyComputeSplineKey(*this, InWorldLocation), ESplineCoordinateSpace::World);
+			if (bUseShallowWaterValues)
+			{
+				WaterPlaneNormal = BakedShallowWaterSim->SimulationData.ComputeShallowWaterSimulationNormalAtPosition(InWorldLocation);
+			}
+			else
+			{
+				// For rivers default to using spline up vector to account for sloping rivers
+				WaterPlaneNormal = WaterSpline->GetUpVectorAtSplineInputKey(Result.LazilyComputeSplineKey(*this, InWorldLocation), ESplineCoordinateSpace::World);
+			}			
 		}
 
 		Result.SetWaterPlaneNormal(WaterPlaneNormal);
@@ -675,7 +708,12 @@ FWaterBodyQueryResult UWaterBodyComponent::QueryWaterInfoClosestToWorldLocation(
 
 		// The better option for computing water depth for ocean and lake is landscape : 
 		const bool bTryUseLandscape = (GetWaterBodyType() == EWaterBodyType::Ocean || GetWaterBodyType() == EWaterBodyType::Lake);
-		if (bTryUseLandscape)
+
+		if (bUseShallowWaterValues)
+		{
+			WaterPlaneDepth = ShallowWaterDepth;
+		}
+		else if (bTryUseLandscape)
 		{
 			TOptional<float> LandscapeHeightOptional;
 			if (ALandscapeProxy* LandscapePtr = FindLandscape())
@@ -775,7 +813,14 @@ FWaterBodyQueryResult UWaterBodyComponent::QueryWaterInfoClosestToWorldLocation(
 		FVector Velocity = FVector::ZeroVector;
 		if (!Result.IsInExclusionVolume())
 		{
-			Velocity = GetWaterVelocityVectorAtSplineInputKey(Result.LazilyComputeSplineKey(*this, InWorldLocation));
+			if (bUseShallowWaterValues)
+			{
+				Velocity = ShallowWaterVelocity;
+			}
+			else
+			{
+				Velocity = GetWaterVelocityVectorAtSplineInputKey(Result.LazilyComputeSplineKey(*this, InWorldLocation));
+			}
 		}
 
 		Result.SetVelocity(Velocity);
diff --git a/Engine/Plugins/Experimental/Water/Source/Runtime/Public/WaterBodyComponent.h b/Engine/Plugins/Experimental/Water/Source/Runtime/Public/WaterBodyComponent.h
--- a/Engine/Plugins/Experimental/Water/Source/Runtime/Public/WaterBodyComponent.h
+++ b/Engine/Plugins/Experimental/Water/Source/Runtime/Public/WaterBodyComponent.h
@@ -432,8 +432,11 @@ public:
 	/** Set the baked shallow water simulation for this water body */
 	void SetBakedShallowWaterSimulation(TObjectPtr<UBakedShallowWaterSimulationComponent> BakedSim) { BakedShallowWaterSim = BakedSim; }
 
-	/** Query for if the baked shallow water sim is valid for use */
-	bool UseBakedShallowWaterSimulationForBuoyancy() const { return bUseBakedSimForPhysics && BakedShallowWaterSim.IsValid() && BakedShallowWaterSim->SimulationData.IsValid(); }
+	/** Set toggle to use baked simulations if they are valid */
+	void SetUseBakedSimulationForQueriesAndPhysics(bool bUseBakedSimulation) { bUseBakedSimForQueriesAndPhysics = bUseBakedSimulation;  }
+	
+	/** Query for if the baked simulations is valid for use */
+	bool UseBakedSimulationForQueriesAndPhysics() const { return bUseBakedSimForQueriesAndPhysics && BakedShallowWaterSim.IsValid() && BakedShallowWaterSim->SimulationData.IsValid(); }
 
 protected:
 	//~ Begin UActorComponent interface.
@@ -692,7 +695,7 @@ protected:
 
 	/**  Override to disable use of the baked shallow water simulation for collisons and other uses */
 	UPROPERTY(Category = BakedSimulation, AdvancedDisplay, EditAnywhere)
-	bool bUseBakedSimForPhysics = true;
+	bool bUseBakedSimForQueriesAndPhysics = true;
 
 #if WITH_EDITORONLY_DATA
 	UPROPERTY()
diff --git a/Engine/Plugins/Experimental/WaterAdvanced/Source/WaterAdvanced/Private/ShallowWaterRiverActor.cpp b/Engine/Plugins/Experimental/WaterAdvanced/Source/WaterAdvanced/Private/ShallowWaterRiverActor.cpp
--- a/Engine/Plugins/Experimental/WaterAdvanced/Source/WaterAdvanced/Private/ShallowWaterRiverActor.cpp
+++ b/Engine/Plugins/Experimental/WaterAdvanced/Source/WaterAdvanced/Private/ShallowWaterRiverActor.cpp
@@ -78,6 +78,8 @@ void UShallowWaterRiverComponent::PostLoad()
 	{
 		RiverSimSystem->Activate();
 	}
+
+	UpdateRenderState();
 }
 
 void UShallowWaterRiverComponent::TickComponent(float DeltaTime, ELevelTick TickType, FActorComponentTickFunction* ThisTickFunction)
@@ -93,30 +95,14 @@ void UShallowWaterRiverComponent::TickComponent(float DeltaTime, ELevelTick Tick
 	{
 		RiverSimSystem->Activate();	
 	}
+#endif
+}
 
-	if (RenderState == EShallowWaterRenderState::WaterComponent)
-	{
-		for (AWaterBody* CurrWaterBody : AllWaterBodies)
-		{
-			TObjectPtr<UWaterBodyComponent> CurrWaterBodyComponent = CurrWaterBody->GetWaterBodyComponent();
-
-			if (CurrWaterBodyComponent != nullptr)
-			{
-				UMaterialInstanceDynamic* WaterMID = CurrWaterBodyComponent->GetWaterMaterialInstance();
-				// override materials on water bodies
-				if (WaterMID != nullptr)
-				{
-					WaterMID->SetTextureParameterValue("BakedWaterSimTex", BakedWaterSurfaceTexture);
-					WaterMID->SetVectorParameterValue("BakedWaterSimLocation", SystemPos);
-					WaterMID->SetVectorParameterValue("BakedWaterSimSize", FVector(WorldGridSize.X, WorldGridSize.Y, 1));
+void UShallowWaterRiverComponent::BeginPlay()
+{
+	Super::BeginPlay();
 
-					//FGuid BakedSimGuid;
-					//WaterMID->SetStaticSwitchParameterValueEditorOnly("UseBakedSim", true, BakedSimGuid);
-				}
-			}
-		}
-	}
-#endif
+	UpdateRenderState();
 }
 
 void UShallowWaterRiverComponent::OnUnregister()
@@ -140,8 +126,6 @@ void UShallowWaterRiverComponent::PostEditChangeProperty(FPropertyChangedEvent&
 	if (PropertyName == GET_MEMBER_NAME_CHECKED(UShallowWaterRiverComponent, RenderState) && RiverSimSystem != nullptr && RiverSimSystem->IsActive())
 	{
 		RiverSimSystem->SetVariableBool(FName("ReadCachedSim"), RenderState == EShallowWaterRenderState::BakedSim);
-
-		// hide/show all the water bodies
 	}
 	else
 	{
@@ -359,6 +343,16 @@ void UShallowWaterRiverComponent::Rebuild()
 				OnWaterInfoTextureArrayCreated(NewWaterInfoTexture);
 			}
 
+			const FVector2D ZoneLocation = FVector2D(WaterZone->GetDynamicWaterInfoCenter());
+
+			const FVector2D ZoneExtent = FVector2D(WaterZone->GetDynamicWaterInfoExtent());
+			const FVector2D WaterHeightExtents = FVector2D(WaterZone->GetWaterHeightExtents());
+			const float GroundZMin = WaterZone->GetGroundZMin();
+
+			RiverSimSystem->SetVariableVec2(FName("WaterZoneLocation"), ZoneLocation);
+			RiverSimSystem->SetVariableVec2(FName("WaterZoneExtent"), ZoneExtent);
+			RiverSimSystem->SetVariableInt(FName("WaterZoneIdx"), WaterZone->GetWaterZoneIndex());
+
 /*
 			// generate high res texture signed distance to water body
 			FVector2D::FReal CellSize = FMath::Max(WorldGridSize.X, WorldGridSize.Y) / ResolutionMaxAxis;
@@ -463,7 +457,7 @@ void UShallowWaterRiverComponent::Rebuild()
 	BakedWaterSurfaceRT = NewObject<UTextureRenderTarget2D>(this, NAME_None, RF_Transient);
 	BakedWaterSurfaceRT->InitAutoFormat(1, 1);
 	RiverSimSystem->SetVariableTextureRenderTarget(FName("SimGridRT"), BakedWaterSurfaceRT);
-	RiverSimSystem->SetVariableBool(FName("ReadCachedSim"), RenderState == EShallowWaterRenderState::BakedSim);
+	RiverSimSystem->SetVariableBool(FName("ReadCachedSim"), RenderState == EShallowWaterRenderState::BakedSim || RenderState == EShallowWaterRenderState::WaterComponentWithBakedSim);
 
 	if (BakedWaterSurfaceTexture != nullptr)
 	{
@@ -573,7 +567,7 @@ bool UShallowWaterRiverComponent::QueryWaterAtSplinePoint(TObjectPtr<AWaterBody>
 
 void UShallowWaterRiverComponent::UpdateRenderState()
 {
-	bool RenderWaterBody = RenderState == EShallowWaterRenderState::WaterComponent;
+	bool RenderWaterBody = RenderState == EShallowWaterRenderState::WaterComponent || RenderState == EShallowWaterRenderState::WaterComponentWithBakedSim;
 
 	RiverSimSystem->SetVisibility(!RenderWaterBody);
 
@@ -586,10 +580,38 @@ void UShallowWaterRiverComponent::UpdateRenderState()
 			CurrWaterBodyComponent->SetVisibility(RenderWaterBody);
 
 			UMaterialInstanceDynamic* WaterMID = CurrWaterBodyComponent->GetWaterMaterialInstance();
-			// override materials on water bodies							
-			WaterMID->SetTextureParameterValue("BakedWaterSimTex", BakedWaterSurfaceTexture);
-			WaterMID->SetVectorParameterValue("BakedWaterSimLocation", SystemPos);
-			WaterMID->SetVectorParameterValue("BakedWaterSimSize", FVector(WorldGridSize.X, WorldGridSize.Y, 1));			
+			UMaterialInstanceDynamic* WaterInfoMID = CurrWaterBodyComponent->GetWaterInfoMaterialInstance();
+			if (RenderState == EShallowWaterRenderState::WaterComponentWithBakedSim)
+			{				
+				// override materials on water bodies							
+				WaterMID->SetTextureParameterValue("BakedWaterSimTex", BakedWaterSurfaceTexture);
+				WaterMID->SetVectorParameterValue("BakedWaterSimLocation", SystemPos);
+				WaterMID->SetVectorParameterValue("BakedWaterSimSize", FVector(WorldGridSize.X, WorldGridSize.Y, 1));
+
+				WaterInfoMID->SetTextureParameterValue("BakedWaterSimTex", BakedWaterSurfaceTexture);
+				WaterInfoMID->SetVectorParameterValue("BakedWaterSimLocation", SystemPos);
+				WaterInfoMID->SetVectorParameterValue("BakedWaterSimSize", FVector(WorldGridSize.X, WorldGridSize.Y, 1));
+			}
+
+			CurrWaterBodyComponent->SetUseBakedSimulationForQueriesAndPhysics(
+				RenderState == EShallowWaterRenderState::WaterComponentWithBakedSim || RenderState == EShallowWaterRenderState::BakedSim);
+
+			/*
+			#todo(dmp): I'd prefer if we could set an editor time only static switch to control using baked sims in the material or not
+			TArray<FMaterialParameterInfo> OutMaterialParameterInfos;
+			TArray<FGuid> Guids;
+			WaterMID->GetAllStaticSwitchParameterInfo(OutMaterialParameterInfos, Guids);
+
+			for (FMaterialParameterInfo& MaterialParameterInfo : OutMaterialParameterInfos)
+			{
+				if (MaterialParameterInfo.Name == "UseBakedSim")
+				{
+					WaterMID->SetStaticSwitchParameterValueEditorOnly(MaterialParameterInfo, RenderState == EShallowWaterRenderState::WaterComponentWithBakedSim);
+				}
+			}
+			*/
+
+			WaterMID->SetScalarParameterValue("UseBakedSimHack", RenderState == EShallowWaterRenderState::WaterComponentWithBakedSim ? 1 : 0);
 		}
 	}
 }
diff --git a/Engine/Plugins/Experimental/WaterAdvanced/Source/WaterAdvanced/Public/ShallowWaterRiverActor.h b/Engine/Plugins/Experimental/WaterAdvanced/Source/WaterAdvanced/Public/ShallowWaterRiverActor.h
--- a/Engine/Plugins/Experimental/WaterAdvanced/Source/WaterAdvanced/Public/ShallowWaterRiverActor.h
+++ b/Engine/Plugins/Experimental/WaterAdvanced/Source/WaterAdvanced/Public/ShallowWaterRiverActor.h
@@ -26,6 +26,7 @@ UENUM(BlueprintType)
 enum EShallowWaterRenderState : int
 {
 	WaterComponent,
+	WaterComponentWithBakedSim,
 	LiveSim,
 	BakedSim
 };
@@ -82,6 +83,7 @@ public:
 	virtual void PostLoad() override;
 
 	virtual void TickComponent(float DeltaTime, ELevelTick TickType, FActorComponentTickFunction* ThisTickFunction) override;
+	virtual void BeginPlay() override;
 
 	virtual void OnUnregister() override;
 
