player based shallow water sims work with baked simulation data when using the water body material

#changelist validated
#virtualized

diff --git a/Engine/Plugins/Experimental/Water/Source/Runtime/Public/BakedShallowWaterSimulationComponent.h b/Engine/Plugins/Experimental/Water/Source/Runtime/Public/BakedShallowWaterSimulationComponent.h
--- a/Engine/Plugins/Experimental/Water/Source/Runtime/Public/BakedShallowWaterSimulationComponent.h
+++ b/Engine/Plugins/Experimental/Water/Source/Runtime/Public/BakedShallowWaterSimulationComponent.h
@@ -13,11 +13,12 @@ struct WATER_API FShallowWaterSimulationGrid
 	GENERATED_BODY()
 
 public:
-	FShallowWaterSimulationGrid(const TArray<FVector4> &InArrayValues, const FIntVector2 &InNumCells, const FVector &InPosition, const FVector2D &InSize)
+	FShallowWaterSimulationGrid(const TArray<FVector4> &InArrayValues, TObjectPtr<UTexture2D> InBakedWaterTexture, const FIntVector2 &InNumCells, const FVector &InPosition, const FVector2D &InSize)
 		: ArrayValues(InArrayValues), 
 		NumCells(InNumCells), 
 		Position(InPosition), 
-		Size(InSize)
+		Size(InSize),
+		BakedTexture(InBakedWaterTexture)
 	{
 		Dx = Size / FVector2D(NumCells.X, NumCells.Y);
 	}
@@ -107,6 +108,10 @@ public:
 	/** World space grid cell size */
 	UPROPERTY()
 	FVector2D Dx;
+	
+	/** Texture for baked water sim data */
+	UPROPERTY()
+	TObjectPtr<UTexture2D> BakedTexture;
 };
 
 UCLASS(meta = (BlueprintSpawnableComponent), MinimalAPI)
diff --git a/Engine/Plugins/Experimental/Water/Source/Runtime/Public/WaterBodyComponent.h b/Engine/Plugins/Experimental/Water/Source/Runtime/Public/WaterBodyComponent.h
--- a/Engine/Plugins/Experimental/Water/Source/Runtime/Public/WaterBodyComponent.h
+++ b/Engine/Plugins/Experimental/Water/Source/Runtime/Public/WaterBodyComponent.h
@@ -427,7 +427,7 @@ public:
 	void OnWaterBodyChanged(bool bShapeOrPositionChanged, bool bWeightmapSettingsChanged = false, bool bUserTriggeredChanged = false) {}
 
 	/** Get the baked shallow water simulation for this water body */
-	UBakedShallowWaterSimulationComponent* GetBakedShallowWaterSimulation() { return BakedShallowWaterSim.Get(); }
+	UBakedShallowWaterSimulationComponent* GetBakedShallowWaterSimulation() const { return BakedShallowWaterSim.Get(); }
 	
 	/** Set the baked shallow water simulation for this water body */
 	void SetBakedShallowWaterSimulation(TObjectPtr<UBakedShallowWaterSimulationComponent> BakedSim) { BakedShallowWaterSim = BakedSim; }
diff --git a/Engine/Plugins/Experimental/WaterAdvanced/Source/WaterAdvanced/Private/ShallowWaterRiverActor.cpp b/Engine/Plugins/Experimental/WaterAdvanced/Source/WaterAdvanced/Private/ShallowWaterRiverActor.cpp
--- a/Engine/Plugins/Experimental/WaterAdvanced/Source/WaterAdvanced/Private/ShallowWaterRiverActor.cpp
+++ b/Engine/Plugins/Experimental/WaterAdvanced/Source/WaterAdvanced/Private/ShallowWaterRiverActor.cpp
@@ -501,7 +501,7 @@ void UShallowWaterRiverComponent::Bake()
 	}
 
 	BakedSim = NewObject<UBakedShallowWaterSimulationComponent>(this, NAME_None, RF_Public);
-	BakedSim->SimulationData = FShallowWaterSimulationGrid(ShallowWaterSimArrayValues, FIntVector2(BakedWaterSurfaceRT->SizeX, BakedWaterSurfaceRT->SizeY), SystemPos, WorldGridSize);
+	BakedSim->SimulationData = FShallowWaterSimulationGrid(ShallowWaterSimArrayValues, BakedWaterSurfaceTexture, FIntVector2(BakedWaterSurfaceRT->SizeX, BakedWaterSurfaceRT->SizeY), SystemPos, WorldGridSize);
 	
 	// set the sim texture on each water body that is in the simulated river.  
 	for (TObjectPtr<AWaterBody > CurrWaterBody : AllWaterBodies)
diff --git a/Engine/Plugins/Experimental/WaterAdvanced/Source/WaterAdvanced/Private/ShallowWaterSubsystem.cpp b/Engine/Plugins/Experimental/WaterAdvanced/Source/WaterAdvanced/Private/ShallowWaterSubsystem.cpp
--- a/Engine/Plugins/Experimental/WaterAdvanced/Source/WaterAdvanced/Private/ShallowWaterSubsystem.cpp
+++ b/Engine/Plugins/Experimental/WaterAdvanced/Source/WaterAdvanced/Private/ShallowWaterSubsystem.cpp
@@ -25,6 +25,7 @@
 #include "GameFramework/SpectatorPawn.h"
 #include "PhysicsEngine/PhysicsAsset.h"
 #include "Engine/AssetManager.h"
+#include "Engine/Texture2D.h"
 
 //UE_DISABLE_OPTIMIZATION_SHIP
 
@@ -515,6 +516,8 @@ void UShallowWaterSubsystem::UpdateGridMovement()
 
 	float BestDistanceSqr = FLT_MAX;
 	TOptional<FVector> BestWaterLocation;
+	
+	TWeakObjectPtr<AWaterBody> BestWaterBody;
 	for (TWeakObjectPtr<AWaterBody> WeakWaterBody : LastOverlappingWaterBodies_Internal)
 	{
 		if (const AWaterBody* WaterBody = WeakWaterBody.Get())
@@ -527,6 +530,7 @@ void UShallowWaterSubsystem::UpdateGridMovement()
 			{
 				BestDistanceSqr = ds;
 				BestWaterLocation = WaterLocation;
+				BestWaterBody = WeakWaterBody;
 			}
 		}
 	}
@@ -548,6 +552,29 @@ void UShallowWaterSubsystem::UpdateGridMovement()
 	ShallowWaterNiagaraSimulation->SetVariableBool(FName("UseDebugRender"), GSWDebugRender == 1);
 	ShallowWaterNiagaraSimulation->SetVariableBool(FName("UseWaterInfoTexture"), GSWUseWaterInfoTexture == 1);
 
+	bool UseBakedSim = false;
+	if (BestWaterBody != nullptr && BestWaterBody->GetWaterBodyComponent())
+	{
+		if (const UWaterBodyComponent* WaterBodyComp = BestWaterBody->GetWaterBodyComponent())
+		{
+			if (WaterBodyComp->UseBakedSimulationForQueriesAndPhysics())
+			{
+				UBakedShallowWaterSimulationComponent* BakedSim = WaterBodyComp->GetBakedShallowWaterSimulation();				
+				UTexture* BakedSimTex = Cast<UTexture>(BakedSim->SimulationData.BakedTexture.Get());
+
+				if (BakedSimTex != nullptr)
+				{
+					ShallowWaterNiagaraSimulation->SetVariableVec3(FName("BakedWaterSimLocation"), BakedSim->SimulationData.Position);
+					ShallowWaterNiagaraSimulation->SetVariableVec2(FName("BakedWaterSimSize"), BakedSim->SimulationData.Size);
+					ShallowWaterNiagaraSimulation->SetVariableTexture(FName("BakedWaterSimTexture"), BakedSimTex);
+
+					UseBakedSim = true;				
+				}
+			}
+		}
+	}
+	ShallowWaterNiagaraSimulation->SetVariableBool(FName("UseBakedSim"), UseBakedSim);
+
 #if ENABLE_DRAW_DEBUG
 	if (GSWDrawWaterSurfaceProjection)
 	{
