Fix invalid call to GetDynamicWaterInfoCenter. Now requires a player index.

diff --git a/Engine/Plugins/Experimental/WaterAdvanced/Source/WaterAdvanced/Private/ShallowWaterRiverActor.cpp b/Engine/Plugins/Experimental/WaterAdvanced/Source/WaterAdvanced/Private/ShallowWaterRiverActor.cpp
--- a/Engine/Plugins/Experimental/WaterAdvanced/Source/WaterAdvanced/Private/ShallowWaterRiverActor.cpp
+++ b/Engine/Plugins/Experimental/WaterAdvanced/Source/WaterAdvanced/Private/ShallowWaterRiverActor.cpp
@@ -343,7 +343,9 @@ void UShallowWaterRiverComponent::Rebuild()
 				OnWaterInfoTextureArrayCreated(NewWaterInfoTexture);
 			}
 
-			const FVector2D ZoneLocation = FVector2D(WaterZone->GetDynamicWaterInfoCenter());
+			// The following index assume that there is no split screen support and will request the position of the first player's water view.
+			const int32 PlayerIndex = 0;
+			const FVector2D ZoneLocation = FVector2D(WaterZone->GetDynamicWaterInfoCenter(PlayerIndex));
 
 			const FVector2D ZoneExtent = FVector2D(WaterZone->GetDynamicWaterInfoExtent());
 			const FVector2D WaterHeightExtents = FVector2D(WaterZone->GetWaterHeightExtents());
