new shallow water subsystem to be used on arbitrary projects

diff --git a/Engine/Plugins/Experimental/WaterAdvanced/Source/WaterAdvanced/Private/BasicShallowWaterSubsystem.cpp b/Engine/Plugins/Experimental/WaterAdvanced/Source/WaterAdvanced/Private/BasicShallowWaterSubsystem.cpp
new file mode 100644
--- /dev/null
+++ b/Engine/Plugins/Experimental/WaterAdvanced/Source/WaterAdvanced/Private/BasicShallowWaterSubsystem.cpp
@@ -0,0 +1,18 @@
+// Copyright Epic Games, Inc. All Rights Reserved.
+
+#include "BasicShallowWaterSubsystem.h"
+
+UBasicShallowWaterSubsystem::UBasicShallowWaterSubsystem()
+{
+}
+
+bool UBasicShallowWaterSubsystem::IsShallowWaterAllowedToInitialize() const
+{
+	TObjectPtr<UShallowWaterSettings> TmpSettings = GetMutableDefault<UShallowWaterSettings>();
+	if (TmpSettings)
+	{
+		return TmpSettings->UseDefaultShallowWaterSubsystem;
+	}
+
+	return false;
+}
\ No newline at end of file
diff --git a/Engine/Plugins/Experimental/WaterAdvanced/Source/WaterAdvanced/Private/ShallowWaterSubsystem.cpp b/Engine/Plugins/Experimental/WaterAdvanced/Source/WaterAdvanced/Private/ShallowWaterSubsystem.cpp
--- a/Engine/Plugins/Experimental/WaterAdvanced/Source/WaterAdvanced/Private/ShallowWaterSubsystem.cpp
+++ b/Engine/Plugins/Experimental/WaterAdvanced/Source/WaterAdvanced/Private/ShallowWaterSubsystem.cpp
@@ -368,12 +368,6 @@ void UShallowWaterSubsystem::InitializeShallowWater()
 
 bool UShallowWaterSubsystem::IsShallowWaterAllowedToInitialize() const
 {
-	TObjectPtr<UShallowWaterSettings> TmpSettings = GetMutableDefault<UShallowWaterSettings>();
-	if (TmpSettings)
-	{
-		return TmpSettings->UseDefaultShallowWaterSubsystem;
-	}
-
 	return false;
 }
 
diff --git a/Engine/Plugins/Experimental/WaterAdvanced/Source/WaterAdvanced/Public/BasicShallowWaterSubsystem.h b/Engine/Plugins/Experimental/WaterAdvanced/Source/WaterAdvanced/Public/BasicShallowWaterSubsystem.h
new file mode 100644
--- /dev/null
+++ b/Engine/Plugins/Experimental/WaterAdvanced/Source/WaterAdvanced/Public/BasicShallowWaterSubsystem.h
@@ -0,0 +1,18 @@
+// Copyright Epic Games, Inc. All Rights Reserved.
+
+#pragma once
+
+#include "ShallowWaterSubsystem.h"
+#include "BasicShallowWaterSubsystem.generated.h"
+
+
+UCLASS(Blueprintable, Transient)
+class WATERADVANCED_API UBasicShallowWaterSubsystem : public UShallowWaterSubsystem
+{
+	GENERATED_BODY()
+
+public:
+	UBasicShallowWaterSubsystem();
+
+	virtual bool IsShallowWaterAllowedToInitialize() const override;	
+};
