- small fixes
- refresh physics test levels
- new river test with many several rapid types
- advect via velocity texture in material for surface wave detail and foam

#virtualized

diff --git a/Engine/Plugins/Experimental/WaterAdvanced/Source/WaterAdvanced/Private/ShallowWaterRiverActor.cpp b/Engine/Plugins/Experimental/WaterAdvanced/Source/WaterAdvanced/Private/ShallowWaterRiverActor.cpp
--- a/Engine/Plugins/Experimental/WaterAdvanced/Source/WaterAdvanced/Private/ShallowWaterRiverActor.cpp
+++ b/Engine/Plugins/Experimental/WaterAdvanced/Source/WaterAdvanced/Private/ShallowWaterRiverActor.cpp
@@ -332,6 +332,8 @@ void UShallowWaterRiverComponent::Rebuild()
 	{
 		if (AWaterZone* WaterZone = CurrWaterBody->GetWaterBodyComponent()->GetWaterZone())
 		{
+			WaterZoneForWIT = WaterZone;
+
 			const TObjectPtr<UTextureRenderTarget2DArray> NewWaterInfoTexture = WaterZone->WaterInfoTextureArray;
 			if (NewWaterInfoTexture == nullptr)
 			{
@@ -342,93 +344,7 @@ void UShallowWaterRiverComponent::Rebuild()
 			{
 				OnWaterInfoTextureArrayCreated(NewWaterInfoTexture);
 			}
-
-			// The following index assume that there is no split screen support and will request the position of the first player's water view.
-			const int32 PlayerIndex = 0;
-			const FVector2D ZoneLocation = FVector2D(WaterZone->GetDynamicWaterInfoCenter(PlayerIndex));
-
-			const FVector2D ZoneExtent = FVector2D(WaterZone->GetDynamicWaterInfoExtent());
-			const FVector2D WaterHeightExtents = FVector2D(WaterZone->GetWaterHeightExtents());
-			const float GroundZMin = WaterZone->GetGroundZMin();
-
-			RiverSimSystem->SetVariableVec2(FName("WaterZoneLocation"), ZoneLocation);
-			RiverSimSystem->SetVariableVec2(FName("WaterZoneExtent"), ZoneExtent);
-			RiverSimSystem->SetVariableInt(FName("WaterZoneIdx"), WaterZone->GetWaterZoneIndex());
-
-/*
-			// generate high res texture signed distance to water body
-			FVector2D::FReal CellSize = FMath::Max(WorldGridSize.X, WorldGridSize.Y) / ResolutionMaxAxis;
-
-			FIntVector2 NumCells;
-			NumCells.X = int32(FMath::FloorToInt(WorldGridSize.X / CellSize));
-			NumCells.Y = int32(FMath::FloorToInt(WorldGridSize.Y / CellSize));
-
-			// Pad grid by 1 voxel if our computed bounding box is too small
-			if (WorldGridSize.X > WorldGridSize.Y && !FMath::IsNearlyEqual(CellSize * NumCells.Y, WorldGridSize.Y))
-			{
-				NumCells.Y++;
-			}
-			else if (WorldGridSize.X < WorldGridSize.Y && !FMath::IsNearlyEqual(CellSize * NumCells.X, WorldGridSize.X))
-			{
-				NumCells.X++;
-			}
-
-			UWaterSplineComponent* SplineComponent = CurrWaterBody->GetWaterBodyComponent()->GetWaterSpline();
-			UWaterSplineMetadata* SplineMetadata = CurrWaterBody->GetWaterBodyComponent()->GetWaterSplineMetadata();
-			FInterpCurveVector PositionCurve = SplineComponent->SplineCurves.Position;
-			const FInterpCurveFloat& WidthCurve = SplineMetadata->RiverWidth;
-			const Chaos::FRigidTransform3 WaterTransform = CurrWaterBody->GetWaterBodyComponent()->GetComponentTransform();
-
-			TArray<FInterpCurvePointVector> &PositionPts = PositionCurve.Points;
-
-			FInterpCurveVector PositionNew;
-			for (FInterpCurvePointVector Pt : PositionPts)
-			{
-				FVector NewPos = Pt.OutVal;
-				NewPos.Z = 0;
-
-				PositionNew.AddPoint(Pt.InVal, NewPos);
-			}			
-
-			TArray<FFloat16> SignedDistanceValues;
-			SignedDistanceValues.SetNum(NumCells.X * NumCells.Y);
-
-
-			for (int y = 0; y < NumCells.Y; ++y) {
-			for (int x = 0; x < NumCells.X; ++x) {
-				FVector2D CurrIndex = FVector2D(x, y);
-				
-
-				const FVector2D UnitPos = (CurrIndex + .5) / FVector2D(NumCells.X, NumCells.Y);
-				const FVector2D LocalPos = (UnitPos - .5) * WorldGridSize;
-				const FVector LocalPos3D = FVector(LocalPos.X, LocalPos.Y, 0);
-				FVector CellWorldPos = LocalPos3D + SystemPos;
-				FVector WaterLocalPos = WaterTransform.InverseTransformPosition(CellWorldPos);
-
-
-				// convert to world space
-				float Dist;
-				const float ClosestSplineKey = PositionCurve.FindNearest(WaterLocalPos, Dist);
-				FVector ClosestPos = PositionCurve.Eval(ClosestSplineKey);
-				float ClosestWidth = WidthCurve.Eval(ClosestSplineKey);
-
-				// snap to xy plane
-				//WaterLocalPos.Z = 0;
-				//ClosestPos.Z = 0;
-				
-				const float SignedDistanceToWater = (WaterLocalPos - ClosestPos).Length() - ClosestWidth;
-
-				SignedDistanceValues[x + NumCells.X * y] = SignedDistanceToWater < 0;
-			}}
-
-			SignedDistanceToSplineTexture = NewObject<UTexture2D>(GetTransientPackage(), NAME_None, RF_Public);
-			SignedDistanceToSplineTexture->Source.Init(NumCells.X, NumCells.Y, 1, 1, TSF_R16F, (const uint8*) SignedDistanceValues.GetData());
-			SignedDistanceToSplineTexture->PowerOfTwoMode = ETexturePowerOfTwoSetting::None;
-			SignedDistanceToSplineTexture->MipGenSettings = TextureMipGenSettings::TMGS_NoMipmaps;			
-			SignedDistanceToSplineTexture->UpdateResource();
-			SignedDistanceToSplineTexture->PostEditChange();
-*/
-
+			
 			break;
 		}
 	}
@@ -574,10 +490,27 @@ void UShallowWaterRiverComponent::OnWaterInfoTextureArrayCreated(const UTextureR
 		ensureMsgf(false, TEXT("UShallowWaterRiverComponent::OnWaterInfoTextureCreated was called with NULL WaterInfoTexture"));
 		return;
 	}
-
+	
 	WaterInfoTexture = InWaterInfoTexture;
 	if (RiverSimSystem)
 	{
+		if (WaterZoneForWIT == nullptr)
+		{
+			ensureMsgf(false, TEXT("UShallowWaterRiverComponent::OnWaterInfoTextureCreated no valid water zone for water info texturep"));
+			return;
+		}
+
+		const int32 PlayerIndex = 0;
+		const FVector2D ZoneLocation = FVector2D(WaterZoneForWIT->GetDynamicWaterInfoCenter(PlayerIndex));
+
+		const FVector2D ZoneExtent = FVector2D(WaterZoneForWIT->GetDynamicWaterInfoExtent());
+		const FVector2D WaterHeightExtents = FVector2D(WaterZoneForWIT->GetWaterHeightExtents());
+		const float GroundZMin = WaterZoneForWIT->GetGroundZMin();
+
+		RiverSimSystem->SetVariableVec2(FName("WaterZoneLocation"), ZoneLocation);
+		RiverSimSystem->SetVariableVec2(FName("WaterZoneExtent"), ZoneExtent);
+		RiverSimSystem->SetVariableInt(FName("WaterZoneIdx"), WaterZoneForWIT->GetWaterZoneIndex());
+
 		UTexture* WITTextureArray = Cast<UTexture>(const_cast<UTextureRenderTarget2DArray*>(WaterInfoTexture.Get()));
 		if (WITTextureArray == nullptr)
 		{
@@ -600,7 +533,10 @@ void UShallowWaterRiverComponent::UpdateRenderState()
 {
 	bool RenderWaterBody = RenderState == EShallowWaterRenderState::WaterComponent || RenderState == EShallowWaterRenderState::WaterComponentWithBakedSim;
 
-	RiverSimSystem->SetVisibility(!RenderWaterBody);
+	if (RiverSimSystem != nullptr)
+	{
+		RiverSimSystem->SetVisibility(!RenderWaterBody);
+	}
 
 	for (AWaterBody* CurrWaterBody : AllWaterBodies)
 	{
diff --git a/Engine/Plugins/Experimental/WaterAdvanced/Source/WaterAdvanced/Public/ShallowWaterRiverActor.h b/Engine/Plugins/Experimental/WaterAdvanced/Source/WaterAdvanced/Public/ShallowWaterRiverActor.h
--- a/Engine/Plugins/Experimental/WaterAdvanced/Source/WaterAdvanced/Public/ShallowWaterRiverActor.h
+++ b/Engine/Plugins/Experimental/WaterAdvanced/Source/WaterAdvanced/Public/ShallowWaterRiverActor.h
@@ -21,6 +21,7 @@
 class UNiagaraComponent;
 class UNiagaraSystem;
 class AWaterBody;
+class AWaterZone;
 
 UENUM(BlueprintType)
 enum EShallowWaterRenderState : int
@@ -93,6 +94,8 @@ public:
 	void Bake();
 
 	virtual void PostEditChangeProperty(FPropertyChangedEvent& PropertyChangedEvent) override;
+
+	UFUNCTION()
 	void OnWaterInfoTextureArrayCreated(const UTextureRenderTarget2DArray* InWaterInfoTexture);
 #endif // WITH_EDITOR
 
@@ -118,6 +121,8 @@ private:
 	bool bIsInitialized;	
 	bool bTickInitialize;
 
+	AWaterZone* WaterZoneForWIT = nullptr;
+
 	UPROPERTY()
 	TSet < TObjectPtr<AWaterBody>> AllWaterBodies;
 
