- workaround for ensure triggered when querying for water zone parameters, although it still happens sometimes
- drive water scattering via velocity

#virtualized

diff --git a/Engine/Plugins/Experimental/WaterAdvanced/Source/WaterAdvanced/Private/ShallowWaterRiverActor.cpp b/Engine/Plugins/Experimental/WaterAdvanced/Source/WaterAdvanced/Private/ShallowWaterRiverActor.cpp
--- a/Engine/Plugins/Experimental/WaterAdvanced/Source/WaterAdvanced/Private/ShallowWaterRiverActor.cpp
+++ b/Engine/Plugins/Experimental/WaterAdvanced/Source/WaterAdvanced/Private/ShallowWaterRiverActor.cpp
@@ -1,11 +1,15 @@
 // Copyright Epic Games, Inc. All Rights Reserved.
 
 #include "ShallowWaterRiverActor.h"
+
 #include "NiagaraComponent.h"
 #include "NiagaraSystem.h"
 #include "NiagaraFunctionLibrary.h"
+
 #include "WaterBodyActor.h"
 #include "WaterSplineComponent.h"
+#include "WaterSubsystem.h"
+
 #include "BakedShallowWaterSimulationComponent.h"
 #include "Engine/TextureRenderTarget2D.h"
 #include "Engine/TextureRenderTarget2DArray.h"
@@ -331,9 +335,7 @@ void UShallowWaterRiverComponent::Rebuild()
 	for (TObjectPtr<AWaterBody > CurrWaterBody : AllWaterBodies)
 	{
 		if (AWaterZone* WaterZone = CurrWaterBody->GetWaterBodyComponent()->GetWaterZone())
-		{
-			WaterZoneForWIT = WaterZone;
-
+		{			
 			const TObjectPtr<UTextureRenderTarget2DArray> NewWaterInfoTexture = WaterZone->WaterInfoTextureArray;
 			if (NewWaterInfoTexture == nullptr)
 			{
@@ -343,7 +345,17 @@ void UShallowWaterRiverComponent::Rebuild()
 			else
 			{
 				OnWaterInfoTextureArrayCreated(NewWaterInfoTexture);
-			}
+			}			
+
+			const int32 PlayerIndex = 0;
+			const FVector2D ZoneLocation = FVector2D(WaterZone->GetDynamicWaterInfoCenter(PlayerIndex));
+			const FVector2D ZoneExtent = FVector2D(WaterZone->GetDynamicWaterInfoExtent());
+			const FVector2D WaterHeightExtents = FVector2D(WaterZone->GetWaterHeightExtents());
+			const float GroundZMin = WaterZone->GetGroundZMin();
+
+			RiverSimSystem->SetVariableVec2(FName("WaterZoneLocation"), ZoneLocation);
+			RiverSimSystem->SetVariableVec2(FName("WaterZoneExtent"), ZoneExtent);
+			RiverSimSystem->SetVariableInt(FName("WaterZoneIdx"), WaterZone->GetWaterZoneIndex());
 			
 			break;
 		}
@@ -364,7 +376,7 @@ void UShallowWaterRiverComponent::Rebuild()
 	RiverSimSystem->SetVariableFloat(FName("SourceAngle"), PI / 2.f + FMath::Acos(SourceDir.Dot(FVector(1,0,0))));
 	
 	// height of the sink box doesn't matter
-	float SinkBoxHeight = 10000000;
+	float SinkBoxHeight = 100000;
 	RiverSimSystem->SetVariablePosition(FName("SinkPos"), SinkPos);
 	RiverSimSystem->SetVariableVec3(FName("SinkSize"), FVector(SinkWidth, SourceSize, SinkBoxHeight));
 	RiverSimSystem->SetVariableFloat(FName("SinkAngle"), PI / 2.f + FMath::Acos(SinkDir.Dot(FVector(1, 0, 0))));
@@ -484,7 +496,7 @@ bool UShallowWaterRiverComponent::QueryWaterAtSplinePoint(TObjectPtr<AWaterBody>
 }
 
 void UShallowWaterRiverComponent::OnWaterInfoTextureArrayCreated(const UTextureRenderTarget2DArray* InWaterInfoTexture)
-{
+{	
 	if (InWaterInfoTexture == nullptr)
 	{
 		ensureMsgf(false, TEXT("UShallowWaterRiverComponent::OnWaterInfoTextureCreated was called with NULL WaterInfoTexture"));
@@ -494,23 +506,6 @@ void UShallowWaterRiverComponent::OnWaterInfoTextureArrayCreated(const UTextureR
 	WaterInfoTexture = InWaterInfoTexture;
 	if (RiverSimSystem)
 	{
-		if (WaterZoneForWIT == nullptr)
-		{
-			ensureMsgf(false, TEXT("UShallowWaterRiverComponent::OnWaterInfoTextureCreated no valid water zone for water info texturep"));
-			return;
-		}
-
-		const int32 PlayerIndex = 0;
-		const FVector2D ZoneLocation = FVector2D(WaterZoneForWIT->GetDynamicWaterInfoCenter(PlayerIndex));
-
-		const FVector2D ZoneExtent = FVector2D(WaterZoneForWIT->GetDynamicWaterInfoExtent());
-		const FVector2D WaterHeightExtents = FVector2D(WaterZoneForWIT->GetWaterHeightExtents());
-		const float GroundZMin = WaterZoneForWIT->GetGroundZMin();
-
-		RiverSimSystem->SetVariableVec2(FName("WaterZoneLocation"), ZoneLocation);
-		RiverSimSystem->SetVariableVec2(FName("WaterZoneExtent"), ZoneExtent);
-		RiverSimSystem->SetVariableInt(FName("WaterZoneIdx"), WaterZoneForWIT->GetWaterZoneIndex());
-
 		UTexture* WITTextureArray = Cast<UTexture>(const_cast<UTextureRenderTarget2DArray*>(WaterInfoTexture.Get()));
 		if (WITTextureArray == nullptr)
 		{
diff --git a/Engine/Plugins/Experimental/WaterAdvanced/Source/WaterAdvanced/Public/ShallowWaterRiverActor.h b/Engine/Plugins/Experimental/WaterAdvanced/Source/WaterAdvanced/Public/ShallowWaterRiverActor.h
--- a/Engine/Plugins/Experimental/WaterAdvanced/Source/WaterAdvanced/Public/ShallowWaterRiverActor.h
+++ b/Engine/Plugins/Experimental/WaterAdvanced/Source/WaterAdvanced/Public/ShallowWaterRiverActor.h
@@ -120,11 +120,9 @@ protected:
 private:
 	bool bIsInitialized;	
 	bool bTickInitialize;
-
-	AWaterZone* WaterZoneForWIT = nullptr;
-
+	
 	UPROPERTY()
-	TSet < TObjectPtr<AWaterBody>> AllWaterBodies;
+	TSet <TObjectPtr<AWaterBody>> AllWaterBodies;
 
 	UPROPERTY()
 	FVector2D WorldGridSize;
