remove unnecessary ensures

diff --git a/Engine/Plugins/Experimental/WaterAdvanced/Source/WaterAdvanced/Private/ShallowWaterSubsystem.cpp b/Engine/Plugins/Experimental/WaterAdvanced/Source/WaterAdvanced/Private/ShallowWaterSubsystem.cpp
--- a/Engine/Plugins/Experimental/WaterAdvanced/Source/WaterAdvanced/Private/ShallowWaterSubsystem.cpp
+++ b/Engine/Plugins/Experimental/WaterAdvanced/Source/WaterAdvanced/Private/ShallowWaterSubsystem.cpp
@@ -997,8 +997,7 @@ void UShallowWaterSubsystem::UpdateCollisionForPendingContexts()
 void UShallowWaterSubsystem::EnableCollisionForContext(const FShallowWaterCollisionContext& Context)
 {
 	if (Context.Component == nullptr)
-	{
-		ensure(false);
+	{		
 		UE_LOG(LogShallowWater, Warning, TEXT("EnableCollisionForContext() - Context Component is nullptr"));
 		return;
 	}
@@ -1124,8 +1123,7 @@ void UShallowWaterSubsystem::DisableCollisionForVehicle(const FShallowWaterColli
 
 	TObjectPtr<USkeletalMeshComponent>* ProxyPointer = VehicleCollisionProxies.Find(Context);
 	if (ProxyPointer == nullptr)
-	{
-		ensure(false);
+	{		
 		return;
 	}
 	
