new bracketing technique for advection

#virtualized

diff --git a/Engine/Plugins/Experimental/WaterAdvanced/Shaders/WaterAdvanced.ush b/Engine/Plugins/Experimental/WaterAdvanced/Shaders/WaterAdvanced.ush
new file mode 100644
--- /dev/null
+++ b/Engine/Plugins/Experimental/WaterAdvanced/Shaders/WaterAdvanced.ush
@@ -0,0 +1,161 @@
+// Copyright Epic Games, Inc. All Rights Reserved.
+
+/*=============================================================================
+WaterAdvanced.ush 
+=============================================================================*/
+
+#pragma once
+
+static const float TWOPI = 2. * PI;
+
+void BracketAngle(float2 Velocity, float BracketAngleDelta, out float2 VelocityDir0, out float2 VelocityDir1, out float BlendAlpha)
+{
+	const float AngleDelta = PI / BracketAngleDelta;
+	const float2x2 RotateByAngleDelta = float2x2(cos(AngleDelta), sin(AngleDelta), -sin(AngleDelta), cos(AngleDelta));
+
+    // direction of input velocity
+	const float Angle = atan2(Velocity.y, Velocity.x) + 2.0 * PI;
+
+    // Snap to min bracketed angle
+	const float Fractional = fmod(Angle, AngleDelta);
+	const float Angle0 = Angle - Fractional;
+    
+    // Compute velocity direction snapped to the bracketed angle
+	VelocityDir0 = float2(cos(Angle0), sin(Angle0));
+
+    // second bracketed direction is just rotated by the angle delta
+	VelocityDir1 = mul(VelocityDir0, RotateByAngleDelta);
+
+    // blend is the distance the velocity heading is to the first bracketed angle
+	BlendAlpha = Fractional / AngleDelta;
+}
+
+  // More general form of bracketing - takes just a float and a bracket size.
+void BracketValue(float Value, float BracketMagnitudeDelta, out float Value0, out float Value1, out float BlendAlpha)
+{	
+	const float Fractional = fmod(Value, BracketMagnitudeDelta);
+	Value0 = Value - Fractional;
+	Value1 = Value0 + BracketMagnitudeDelta;
+	BlendAlpha = Fractional / BracketMagnitudeDelta;
+}
+
+float2 CalculateUV(float2 vAxis, float2 position, float Time, float Speed)
+{
+	// second 2d vector perpendicular to the vAxis passed in
+	const float2 uAxis = float2(-vAxis.y, vAxis.x);
+	
+	// project the position onto the basis vectors
+	float2 uv = float2(dot(position, uAxis), dot(position, vAxis));
+    
+	// animated along one of the basis axes
+	uv.y -= Time * Speed;
+	
+	return uv;
+}
+
+float2 ComputeMipLevel(float2 uv, int2 sz)
+{	
+	const float2 dxval = ddx(uv) * sz;
+	const float2 dyval = ddy(uv) * sz;
+
+	const float DDXLengthSquared = dot(dxval, dxval);
+	const float DDYLengthSquared = dot(dyval, dyval);
+	const float MaxLengthSquared = max(DDXLengthSquared, DDYLengthSquared);
+	
+	return 0.5f * log2(max(MaxLengthSquared, 1e-8f));
+}
+
+void Compute3Phases(const float PulseLength, const float Time, 
+	out float lerp0, out float lerp1, out float lerp2,
+	out float t0, out float t1, out float t2)
+{
+	const float HalfPulseLength = .5 * PulseLength;
+	const float TimeScale = TWOPI / PulseLength;
+
+	const int NumPhases = 3;
+
+	const float lerp_t0 = TimeScale * Time + 0. * TWOPI / NumPhases;
+	const float lerp_t1 = TimeScale * Time + 1. * TWOPI / NumPhases;
+	const float lerp_t2 = TimeScale * Time + 2. * TWOPI / NumPhases;
+
+	lerp0 = (.5 - .5 * cos(lerp_t0)) / 1.5;
+	lerp1 = (.5 - .5 * cos(lerp_t1)) / 1.5;
+	lerp2 = (.5 - .5 * cos(lerp_t2)) / 1.5;
+
+	t0 = (Time + PulseLength * 0. / NumPhases) % PulseLength;
+	t1 = (Time + PulseLength * 1. / NumPhases) % PulseLength;
+	t2 = (Time + PulseLength * 2. / NumPhases) % PulseLength;
+}
+
+void Compute2Phases(const float PulseLength, const float Time,
+	out float lerpV,
+	out float t0, out float t1, out float2 CycleCount)
+{	
+	const float HalfPulseLength = .5 * PulseLength;
+	t0 = Time % PulseLength;
+	t1 = (Time + HalfPulseLength) % PulseLength;
+
+	lerpV = -abs(t1 * 2. / PulseLength - 1.) + 1.;
+	
+	CycleCount = float2(floor(Time / PulseLength), floor((Time + HalfPulseLength) / PulseLength));
+}
+
+float4 AdvectionInterpolation2(float2 UV, float2 Velocity, float t0, float t1, float lerpV,
+  float AdvectedTextureScale, float2 TexMipLevel, Texture2D<float4> Tex, sampler TexSampler, float2 CycleCount)
+{
+	// advect uv, offset position somewhat randomly by the cycle count
+	const float2 AdvectedUV0 = UV - Velocity * t0 + CycleCount.x * 123;
+	const float2 AdvectedUV1 = UV - Velocity * t1 + CycleCount.y * 123;
+
+	// scale uv
+	const float2 uv0 = AdvectedUV0 * AdvectedTextureScale;
+	const float2 uv1 = AdvectedUV1 * AdvectedTextureScale;
+
+	// compute mip levels
+	if (TexMipLevel.x < -1e-5 || TexMipLevel.y < 1e-5)
+	{
+		int2 TexDims;
+		int Levels;
+		Tex.GetDimensions(0, TexDims.x, TexDims.y, Levels);
+		TexMipLevel.x = ComputeMipLevel(uv0, TexDims);
+		TexMipLevel.y = ComputeMipLevel(uv1, TexDims);
+	}
+
+	// sample the texture for the two phases
+	const float4 TextureSample0 = Tex.SampleLevel(TexSampler, uv0, TexMipLevel.x);
+	const float4 TextureSample1 = Tex.SampleLevel(TexSampler, uv1, TexMipLevel.y);
+
+	// lerp according to phases
+	return lerp(TextureSample0, TextureSample1, lerpV);	
+}
+
+float4 AdvectionInterpolation(float2 UV, float2 Velocity, float t0, float t1, float t2, float lerp0, float lerp1, float lerp2,
+  float AdvectedTextureScale, float TexMipLevel, Texture2D<float4> Tex, sampler TexSampler)
+{
+	// advect uvs
+	const float2 AdvectedUV0 = UV - Velocity * t0;
+	const float2 AdvectedUV1 = UV - Velocity * t1;
+	const float2 AdvectedUV2 = UV - Velocity * t2;
+
+	// scale uvs
+	const float2 uv0 = AdvectedUV0 * AdvectedTextureScale;
+	const float2 uv1 = AdvectedUV1 * AdvectedTextureScale;
+	const float2 uv2 = AdvectedUV2 * AdvectedTextureScale;
+
+	// compute mips
+	if (TexMipLevel < -1e-5)
+	{	
+		int2 TexDims;
+		int Levels;
+		Tex.GetDimensions(0, TexDims.x, TexDims.y, Levels);
+		TexMipLevel = ComputeMipLevel(uv0, TexDims);
+	}
+		
+	// sample texture for each phase
+	const float4 TextureSample0 = Tex.SampleLevel(TexSampler, uv0, TexMipLevel);
+	const float4 TextureSample1 = Tex.SampleLevel(TexSampler, uv1, TexMipLevel);
+	const float4 TextureSample2 = Tex.SampleLevel(TexSampler, uv2, TexMipLevel);
+	
+	// blend according to phase
+	return TextureSample0 * lerp0 + TextureSample1 * lerp1 + TextureSample2 * lerp2;
+}
\ No newline at end of file
diff --git a/Engine/Plugins/Experimental/WaterAdvanced/Source/WaterAdvanced/Private/WaterAdvanced.cpp b/Engine/Plugins/Experimental/WaterAdvanced/Source/WaterAdvanced/Private/WaterAdvanced.cpp
--- a/Engine/Plugins/Experimental/WaterAdvanced/Source/WaterAdvanced/Private/WaterAdvanced.cpp
+++ b/Engine/Plugins/Experimental/WaterAdvanced/Source/WaterAdvanced/Private/WaterAdvanced.cpp
@@ -3,6 +3,7 @@
 #include "WaterAdvanced.h"
 #include "ShallowWaterRiverDetails.h"
 #include "ShallowWaterRiverActor.h"
+#include "Interfaces/IPluginManager.h"
 
 #if WITH_EDITOR
 #include "PropertyCustomizationHelpers.h"
@@ -20,6 +21,8 @@ void FWaterAdvancedModule::StartupModule()
 		UShallowWaterRiverComponent::StaticClass()->GetFName(),
 		FOnGetDetailCustomizationInstance::CreateStatic(&FShallowWaterRiverDetails::MakeInstance));
 #endif
+	FString PluginShaderDir = FPaths::Combine(IPluginManager::Get().FindPlugin(TEXT("WaterAdvanced"))->GetBaseDir(), TEXT("Shaders"));
+	AddShaderSourceDirectoryMapping(TEXT("/Plugin/Experimental/WaterAdvanced"), PluginShaderDir);
 }
 
 void FWaterAdvancedModule::ShutdownModule()
