fix impacts because of bad subsystem initialization

diff --git a/Engine/Plugins/Experimental/WaterAdvanced/Source/WaterAdvanced/Private/BasicShallowWaterSubsystem.cpp b/Engine/Plugins/Experimental/WaterAdvanced/Source/WaterAdvanced/Private/BasicShallowWaterSubsystem.cpp
--- a/Engine/Plugins/Experimental/WaterAdvanced/Source/WaterAdvanced/Private/BasicShallowWaterSubsystem.cpp
+++ b/Engine/Plugins/Experimental/WaterAdvanced/Source/WaterAdvanced/Private/BasicShallowWaterSubsystem.cpp
@@ -6,6 +6,19 @@ UBasicShallowWaterSubsystem::UBasicShallowWaterSubsystem()
 {
 }
 
+bool UBasicShallowWaterSubsystem::ShouldCreateSubsystem(UObject* Outer) const
+{
+	TObjectPtr<UShallowWaterSettings> TmpSettings = GetMutableDefault<UShallowWaterSettings>();
+	if (TmpSettings && !TmpSettings->UseDefaultShallowWaterSubsystem)
+	{
+		return false;
+	}
+	else
+	{
+		return Super::ShouldCreateSubsystem(Outer);
+	}	
+}
+
 bool UBasicShallowWaterSubsystem::IsShallowWaterAllowedToInitialize() const
 {
 	TObjectPtr<UShallowWaterSettings> TmpSettings = GetMutableDefault<UShallowWaterSettings>();
diff --git a/Engine/Plugins/Experimental/WaterAdvanced/Source/WaterAdvanced/Public/BasicShallowWaterSubsystem.h b/Engine/Plugins/Experimental/WaterAdvanced/Source/WaterAdvanced/Public/BasicShallowWaterSubsystem.h
--- a/Engine/Plugins/Experimental/WaterAdvanced/Source/WaterAdvanced/Public/BasicShallowWaterSubsystem.h
+++ b/Engine/Plugins/Experimental/WaterAdvanced/Source/WaterAdvanced/Public/BasicShallowWaterSubsystem.h
@@ -14,5 +14,6 @@ class WATERADVANCED_API UBasicShallowWaterSubsystem : public UShallowWaterSubsys
 public:
 	UBasicShallowWaterSubsystem();
 
+	virtual bool ShouldCreateSubsystem(UObject* Outer) const override;
 	virtual bool IsShallowWaterAllowedToInitialize() const override;	
 };
