support water body materials w/ baked shallow water sims

#virtualized

diff --git a/Engine/Plugins/Experimental/WaterAdvanced/Source/WaterAdvanced/Private/ShallowWaterRiverActor.cpp b/Engine/Plugins/Experimental/WaterAdvanced/Source/WaterAdvanced/Private/ShallowWaterRiverActor.cpp
--- a/Engine/Plugins/Experimental/WaterAdvanced/Source/WaterAdvanced/Private/ShallowWaterRiverActor.cpp
+++ b/Engine/Plugins/Experimental/WaterAdvanced/Source/WaterAdvanced/Private/ShallowWaterRiverActor.cpp
@@ -395,6 +395,20 @@ void UShallowWaterRiverComponent::Rebuild()
 		RiverSimSystem->SetVariableTexture(FName("BakedSimTexture"), BakedWaterSurfaceTexture);
 	}
 
+	FIntVector2 NumCells = FIntVector2(BakedWaterSurfaceTexture->GetSizeX(), BakedWaterSurfaceTexture->GetSizeY());
+
+	BakedWaterSurfaceNormalRT = NewObject<UTextureRenderTarget2D>(this);	
+	BakedWaterSurfaceNormalRT->InitAutoFormat(1, 1);
+	RiverSimSystem->SetVariableTextureRenderTarget(FName("NormalRT"), BakedWaterSurfaceNormalRT);
+	
+	FoamRT = NewObject<UTextureRenderTarget2D>(this);
+	FoamRT->InitAutoFormat(1, 1);
+	RiverSimSystem->SetVariableTextureRenderTarget(FName("FoamRT"), FoamRT);
+	
+	NormalDetailRT = NewObject<UTextureRenderTarget2D>(this);	
+	NormalDetailRT->InitAutoFormat(1, 1);
+	RiverSimSystem->SetVariableTextureRenderTarget(FName("NormalDetailRT"), NormalDetailRT);
+
 	bIsInitialized = true;
 }
 
@@ -551,6 +565,24 @@ void UShallowWaterRiverComponent::UpdateRenderState()
 				WaterMID->SetVectorParameterValue("BakedWaterSimLocation", SystemPos);
 				WaterMID->SetVectorParameterValue("BakedWaterSimSize", FVector(WorldGridSize.X, WorldGridSize.Y, 1));
 
+				WaterMID->SetTextureParameterValue("BakedWaterSimNormalTex", BakedWaterSurfaceNormalRT);
+				WaterMID->SetTextureParameterValue("NormalDetailTex", NormalDetailRT);
+				WaterMID->SetTextureParameterValue("FoamTex", FoamRT);
+
+				// 
+				
+				float Dx;
+				if (WorldGridSize.X > WorldGridSize.Y)
+				{
+					Dx = WorldGridSize.X / ResolutionMaxAxis;
+				}
+				else					
+				{
+					Dx = WorldGridSize.Y / ResolutionMaxAxis;
+				}				
+				WaterMID->SetScalarParameterValue("BakedWaterSimDx", Dx);
+
+
 				WaterInfoMID->SetTextureParameterValue("BakedWaterSimTex", BakedWaterSurfaceTexture);
 				WaterInfoMID->SetVectorParameterValue("BakedWaterSimLocation", SystemPos);
 				WaterInfoMID->SetVectorParameterValue("BakedWaterSimSize", FVector(WorldGridSize.X, WorldGridSize.Y, 1));
diff --git a/Engine/Plugins/Experimental/WaterAdvanced/Source/WaterAdvanced/Public/ShallowWaterRiverActor.h b/Engine/Plugins/Experimental/WaterAdvanced/Source/WaterAdvanced/Public/ShallowWaterRiverActor.h
--- a/Engine/Plugins/Experimental/WaterAdvanced/Source/WaterAdvanced/Public/ShallowWaterRiverActor.h
+++ b/Engine/Plugins/Experimental/WaterAdvanced/Source/WaterAdvanced/Public/ShallowWaterRiverActor.h
@@ -129,6 +129,15 @@ private:
 
 	UPROPERTY()
 	FVector SystemPos;
+	
+	UPROPERTY()
+	TObjectPtr<UTextureRenderTarget2D> BakedWaterSurfaceNormalRT;
+
+	UPROPERTY()
+	TObjectPtr<UTextureRenderTarget2D> FoamRT;
+
+	UPROPERTY()
+	TObjectPtr<UTextureRenderTarget2D> NormalDetailRT;
 };
 
 UCLASS(BlueprintType, HideCategories = (Physics, Replication, Input, Collision))
