material supports water sim
reset button turns to live mode
particles render in water component rendering modes

#changelist validated
#virtualized

diff --git a/Engine/Plugins/Experimental/WaterAdvanced/Source/WaterAdvanced/Private/ShallowWaterRiverActor.cpp b/Engine/Plugins/Experimental/WaterAdvanced/Source/WaterAdvanced/Private/ShallowWaterRiverActor.cpp
--- a/Engine/Plugins/Experimental/WaterAdvanced/Source/WaterAdvanced/Private/ShallowWaterRiverActor.cpp
+++ b/Engine/Plugins/Experimental/WaterAdvanced/Source/WaterAdvanced/Private/ShallowWaterRiverActor.cpp
@@ -129,7 +129,7 @@ void UShallowWaterRiverComponent::PostEditChangeProperty(FPropertyChangedEvent&
 	// this should go before rebuild not after...something is wrong
 	if (PropertyName == GET_MEMBER_NAME_CHECKED(UShallowWaterRiverComponent, RenderState) && RiverSimSystem != nullptr && RiverSimSystem->IsActive())
 	{
-		RiverSimSystem->SetVariableBool(FName("ReadCachedSim"), RenderState == EShallowWaterRenderState::BakedSim);
+		RiverSimSystem->SetVariableBool(FName("ReadCachedSim"), RenderState == EShallowWaterRenderState::BakedSim || RenderState == EShallowWaterRenderState::WaterComponentWithBakedSim);
 	}
 	else
 	{
@@ -544,8 +544,8 @@ void UShallowWaterRiverComponent::UpdateRenderState()
 	bool RenderWaterBody = RenderState == EShallowWaterRenderState::WaterComponent || RenderState == EShallowWaterRenderState::WaterComponentWithBakedSim;
 
 	if (RiverSimSystem != nullptr)
-	{
-		RiverSimSystem->SetVisibility(!RenderWaterBody);
+	{		
+		RiverSimSystem->SetVariableBool(FName("RenderWaterSurface"), !RenderWaterBody);
 	}
 
 	for (AWaterBody* CurrWaterBody : AllWaterBodies)
diff --git a/Engine/Plugins/Experimental/WaterAdvanced/Source/WaterAdvanced/Private/ShallowWaterRiverDetails.cpp b/Engine/Plugins/Experimental/WaterAdvanced/Source/WaterAdvanced/Private/ShallowWaterRiverDetails.cpp
--- a/Engine/Plugins/Experimental/WaterAdvanced/Source/WaterAdvanced/Private/ShallowWaterRiverDetails.cpp
+++ b/Engine/Plugins/Experimental/WaterAdvanced/Source/WaterAdvanced/Private/ShallowWaterRiverDetails.cpp
@@ -135,6 +135,8 @@ FReply FShallowWaterRiverDetails::OnResetSelectedSystem()
 					UShallowWaterRiverComponent* ShallowWaterComponent = Cast<UShallowWaterRiverComponent>(AC);
 					if (ShallowWaterComponent)
 					{
+						ShallowWaterComponent->RenderState = EShallowWaterRenderState::LiveSim;
+						ShallowWaterComponent->UpdateRenderState();
 						ShallowWaterComponent->Rebuild();
 						ShallowWaterComponent->ReregisterComponent();
 					}
@@ -142,6 +144,8 @@ FReply FShallowWaterRiverDetails::OnResetSelectedSystem()
 			}
 			else if (UShallowWaterRiverComponent* ShallowWaterComponent = Cast<UShallowWaterRiverComponent>(SelectedObjects[Idx].Get()))
 			{
+				ShallowWaterComponent->RenderState = EShallowWaterRenderState::LiveSim;
+				ShallowWaterComponent->UpdateRenderState();
 				ShallowWaterComponent->Rebuild();
 				ShallowWaterComponent->ReregisterComponent();
 			}
