fix null ptr

diff --git a/Engine/Plugins/Experimental/WaterAdvanced/Source/WaterAdvanced/Private/ShallowWaterRiverActor.cpp b/Engine/Plugins/Experimental/WaterAdvanced/Source/WaterAdvanced/Private/ShallowWaterRiverActor.cpp
--- a/Engine/Plugins/Experimental/WaterAdvanced/Source/WaterAdvanced/Private/ShallowWaterRiverActor.cpp
+++ b/Engine/Plugins/Experimental/WaterAdvanced/Source/WaterAdvanced/Private/ShallowWaterRiverActor.cpp
@@ -393,21 +393,22 @@ void UShallowWaterRiverComponent::Rebuild()
 	if (BakedWaterSurfaceTexture != nullptr)
 	{
 		RiverSimSystem->SetVariableTexture(FName("BakedSimTexture"), BakedWaterSurfaceTexture);
-	}
+	
 
-	FIntVector2 NumCells = FIntVector2(BakedWaterSurfaceTexture->GetSizeX(), BakedWaterSurfaceTexture->GetSizeY());
+		FIntVector2 NumCells = FIntVector2(BakedWaterSurfaceTexture->GetSizeX(), BakedWaterSurfaceTexture->GetSizeY());
 
-	BakedWaterSurfaceNormalRT = NewObject<UTextureRenderTarget2D>(this);	
-	BakedWaterSurfaceNormalRT->InitAutoFormat(1, 1);
-	RiverSimSystem->SetVariableTextureRenderTarget(FName("NormalRT"), BakedWaterSurfaceNormalRT);
+		BakedWaterSurfaceNormalRT = NewObject<UTextureRenderTarget2D>(this);	
+		BakedWaterSurfaceNormalRT->InitAutoFormat(1, 1);
+		RiverSimSystem->SetVariableTextureRenderTarget(FName("NormalRT"), BakedWaterSurfaceNormalRT);
 	
-	FoamRT = NewObject<UTextureRenderTarget2D>(this);
-	FoamRT->InitAutoFormat(1, 1);
-	RiverSimSystem->SetVariableTextureRenderTarget(FName("FoamRT"), FoamRT);
+		FoamRT = NewObject<UTextureRenderTarget2D>(this);
+		FoamRT->InitAutoFormat(1, 1);
+		RiverSimSystem->SetVariableTextureRenderTarget(FName("FoamRT"), FoamRT);
 	
-	NormalDetailRT = NewObject<UTextureRenderTarget2D>(this);	
-	NormalDetailRT->InitAutoFormat(1, 1);
-	RiverSimSystem->SetVariableTextureRenderTarget(FName("NormalDetailRT"), NormalDetailRT);
+		NormalDetailRT = NewObject<UTextureRenderTarget2D>(this);	
+		NormalDetailRT->InitAutoFormat(1, 1);
+		RiverSimSystem->SetVariableTextureRenderTarget(FName("NormalDetailRT"), NormalDetailRT);
+	}
 
 	bIsInitialized = true;
 }
