fix initialization issues and some rare crashes

diff --git a/Engine/Plugins/Experimental/WaterAdvanced/Source/WaterAdvanced/Private/ShallowWaterRiverActor.cpp b/Engine/Plugins/Experimental/WaterAdvanced/Source/WaterAdvanced/Private/ShallowWaterRiverActor.cpp
--- a/Engine/Plugins/Experimental/WaterAdvanced/Source/WaterAdvanced/Private/ShallowWaterRiverActor.cpp
+++ b/Engine/Plugins/Experimental/WaterAdvanced/Source/WaterAdvanced/Private/ShallowWaterRiverActor.cpp
@@ -38,24 +38,7 @@ UShallowWaterRiverComponent::UShallowWaterRiverComponent(const FObjectInitialize
 	ResolutionMaxAxis = 512;
 	SourceSize = 1000;
 
-	// #todo(dmp): default river system should be set here
-	//hap = LoadObject<UNiagaraSystem>(nullptr, TEXT("/WaterAdvanced/Niagara/Systems/Grid2D_SW_River.Grid2D_SW_River"));
-
-#if WITH_EDITOR
-	// Add all landscapes by default
-	/*
-	if (BottomContourActors.IsEmpty())
-	{
-		TArray<AActor*> FoundActors;
-		UGameplayStatics::GetAllActorsOfClass(GetWorld(), ALandscape::StaticClass(), FoundActors);
-
-		for (AActor* CurrActor : FoundActors)
-		{
-			BottomContourActors.Add(CurrActor);
-		}
-	}
-	*/
-
+#if WITH_EDITOR	
 	// start with one empty element
 	if (BottomContourActors.IsEmpty())
 	{
@@ -90,12 +73,12 @@ void UShallowWaterRiverComponent::TickComponent(float DeltaTime, ELevelTick Tick
 {
 #if WITH_EDITOR
 	// lots of tick ordering issues, so we try to initialize on the first tick too
-	if (RiverSimSystem == nullptr || (RenderState == EShallowWaterRenderState::LiveSim && !bIsInitialized && !bTickInitialize))
+	if (!bTickInitialize && (RiverSimSystem == nullptr || (RenderState == EShallowWaterRenderState::LiveSim && !bIsInitialized)))
 	{
 		bTickInitialize = true;
 		Rebuild();
 	}
-	else
+	else if (bIsInitialized)
 	{
 		RiverSimSystem->Activate();	
 	}
@@ -144,6 +127,11 @@ void UShallowWaterRiverComponent::PostEditChangeProperty(FPropertyChangedEvent&
 
 void UShallowWaterRiverComponent::Rebuild()
 {	
+	if (NiagaraRiverSimulation == nullptr)
+	{
+		NiagaraRiverSimulation = LoadObject<UNiagaraSystem>(nullptr, TEXT("/WaterAdvanced/Niagara/Systems/Grid2D_SW_River.Grid2D_SW_River"));
+	}
+
 	if (RiverSimSystem != nullptr)
 	{
 		RiverSimSystem->SetActive(false);
@@ -417,6 +405,13 @@ void UShallowWaterRiverComponent::Bake()
 {	
 	EObjectFlags TextureObjectFlags = EObjectFlags::RF_Public;
 
+	if (!BakedWaterSurfaceRT)
+	{
+		UE_LOG(LogShallowWater, Warning, TEXT("UShallowWaterRiverComponent::Bake() - No simulation to bake"));
+		
+		return;
+	}
+
 	BakedWaterSurfaceTexture = BakedWaterSurfaceRT->ConstructTexture2D(this, "BakedRiverTexture", TextureObjectFlags);
 
 	RiverSimSystem->SetVariableTexture(FName("BakedSimTexture"), BakedWaterSurfaceTexture);	
