set water sim texture on all water mids for all water bodies in a water zone once they are encountered
correctly compute bounds on water body sim
turn off dissipation outside of river channel

#virtualized

diff --git a/Engine/Plugins/Experimental/WaterAdvanced/Source/WaterAdvanced/Private/ShallowWaterSubsystem.cpp b/Engine/Plugins/Experimental/WaterAdvanced/Source/WaterAdvanced/Private/ShallowWaterSubsystem.cpp
--- a/Engine/Plugins/Experimental/WaterAdvanced/Source/WaterAdvanced/Private/ShallowWaterSubsystem.cpp
+++ b/Engine/Plugins/Experimental/WaterAdvanced/Source/WaterAdvanced/Private/ShallowWaterSubsystem.cpp
@@ -203,7 +203,28 @@ void UShallowWaterSubsystem::Tick(float DeltaTime)
 			{
 				continue;
 			}
-			TryUpdateWaterBodyMIDParameters(Water->GetWaterBodyComponent());
+
+			// loop over the water bodies for the current water zone and make sure each one has all of the materials setup with 
+			// references to the simulation texture(s)
+			if (const UWaterBodyComponent* WaterBodyComp = Water->GetWaterBodyComponent())
+			{
+				if (const AWaterZone* WaterZone = WaterBodyComp->GetWaterZone())
+				{
+					// Update all water bodies in the water zone only if it hasn't already been done
+					if (!WaterZonesWithMaterialsInitialized.Find(WaterZone))
+					{
+						WaterZone->ForEachWaterBodyComponent([this](UWaterBodyComponent* WaterBodyComponent)
+							{
+								TryUpdateWaterBodyMIDParameters(WaterBodyComponent);
+
+								return true;
+							}
+						);
+
+						WaterZonesWithMaterialsInitialized.Add(WaterZone);
+					}
+				}
+			}
 		}
 	}
 
@@ -335,6 +356,8 @@ void UShallowWaterSubsystem::InitializeShallowWater()
 
 	bIsShallowWaterInitialized = true;
 	
+	WaterZonesWithMaterialsInitialized.Empty();
+	
 	CreateRTs();
 	
 	ShallowWaterNiagaraSimulation->SetUsingAbsoluteRotation(true);
diff --git a/Engine/Plugins/Experimental/WaterAdvanced/Source/WaterAdvanced/Public/ShallowWaterSubsystem.h b/Engine/Plugins/Experimental/WaterAdvanced/Source/WaterAdvanced/Public/ShallowWaterSubsystem.h
--- a/Engine/Plugins/Experimental/WaterAdvanced/Source/WaterAdvanced/Public/ShallowWaterSubsystem.h
+++ b/Engine/Plugins/Experimental/WaterAdvanced/Source/WaterAdvanced/Public/ShallowWaterSubsystem.h
@@ -15,6 +15,7 @@
 #include "ShallowWaterSubsystem.generated.h"
 
 class AWaterBody;
+class AWaterZone;
 class UWaterBodyComponent;
 class UNiagaraComponent;
 class APawn;
@@ -265,6 +266,8 @@ private:
 
 	bool bInitializationAsyncLoadsAttempted = false;
 
+	TSet<const AWaterZone*> WaterZonesWithMaterialsInitialized;
+
 	struct PendingImpact
 	{
 		FVector ImpactPosition;
