modify spline river collision geometry to fully contain baked water sim rivers

#virtualized

diff --git a/Engine/Plugins/Experimental/WaterAdvanced/Source/WaterAdvanced/Private/ShallowWaterRiverActor.cpp b/Engine/Plugins/Experimental/WaterAdvanced/Source/WaterAdvanced/Private/ShallowWaterRiverActor.cpp
--- a/Engine/Plugins/Experimental/WaterAdvanced/Source/WaterAdvanced/Private/ShallowWaterRiverActor.cpp
+++ b/Engine/Plugins/Experimental/WaterAdvanced/Source/WaterAdvanced/Private/ShallowWaterRiverActor.cpp
@@ -14,6 +14,9 @@
 #include "Engine/TextureRenderTarget2D.h"
 #include "Engine/TextureRenderTarget2DArray.h"
 #include "Materials/MaterialInstanceDynamic.h"
+#include "Components/SplineMeshComponent.h"
+#include "PhysicsEngine/ConvexElem.h"
+#include "PhysicsEngine/BodySetup.h"
 
 #include "TextureResource.h"
 #include "ShallowWaterCommon.h"
@@ -402,24 +405,24 @@ void UShallowWaterRiverComponent::Rebuild()
 }
 
 void UShallowWaterRiverComponent::Bake()
-{	
+{
 	EObjectFlags TextureObjectFlags = EObjectFlags::RF_Public;
 
 	if (!BakedWaterSurfaceRT)
 	{
 		UE_LOG(LogShallowWater, Warning, TEXT("UShallowWaterRiverComponent::Bake() - No simulation to bake"));
-		
+
 		return;
 	}
 
 	BakedWaterSurfaceTexture = BakedWaterSurfaceRT->ConstructTexture2D(this, "BakedRiverTexture", TextureObjectFlags);
 
-	RiverSimSystem->SetVariableTexture(FName("BakedSimTexture"), BakedWaterSurfaceTexture);	
+	RiverSimSystem->SetVariableTexture(FName("BakedSimTexture"), BakedWaterSurfaceTexture);
 
 	// Readback to get the river texture values as an array
 	TArray<FFloat16Color> TmpShallowWaterSimArrayValues;
 	BakedWaterSurfaceRT->GameThread_GetRenderTargetResource()->ReadFloat16Pixels(TmpShallowWaterSimArrayValues);
-		
+
 	TArray<FVector4> ShallowWaterSimArrayValues;
 	ShallowWaterSimArrayValues.Empty();
 	ShallowWaterSimArrayValues.AddZeroed(TmpShallowWaterSimArrayValues.Num());
@@ -431,7 +434,7 @@ void UShallowWaterRiverComponent::Bake()
 		const float WaterHeight = Val.R;
 		const float WaterDepth = Val.G;
 		const FVector2D WaterVelocity(Val.B, Val.A);
-		
+
 		FVector4 FloatVal;
 		FloatVal.X = WaterHeight;
 		FloatVal.Y = WaterDepth;
@@ -443,13 +446,126 @@ void UShallowWaterRiverComponent::Bake()
 
 	BakedSim = NewObject<UBakedShallowWaterSimulationComponent>(this, NAME_None, RF_Public);
 	BakedSim->SimulationData = FShallowWaterSimulationGrid(ShallowWaterSimArrayValues, BakedWaterSurfaceTexture, FIntVector2(BakedWaterSurfaceRT->SizeX, BakedWaterSurfaceRT->SizeY), SystemPos, WorldGridSize);
-	
+		
+	// compute the maximum water height for each convex in each water body simulated by this river
+	// we use this to modify the collision geometry so it fully encompasses the baked water sim
+	TMap<FKConvexElem*, float> ConvexToMaxHeight;
+	for (int32 y = 0; y < BakedWaterSurfaceRT->SizeY; ++y) {
+	for (int32 x = 0; x < BakedWaterSurfaceRT->SizeX; ++x) {
+		FVector WorldPos = BakedSim->SimulationData.IndexToWorld(FIntVector2(x, y));
+
+		FVector Vel;
+		float Height, Depth;
+		BakedSim->SimulationData.QueryShallowWaterSimulationAtIndex(FIntVector2(x, y), Vel, Height, Depth);
+		WorldPos.Z = Height;
+				
+		if (Depth > 1e-5)
+		{
+			for (TObjectPtr<AWaterBody > CurrWaterBody : AllWaterBodies)
+			{						
+				TObjectPtr<UWaterBodyComponent> CurrWaterBodyComponent = CurrWaterBody->GetWaterBodyComponent();
+
+				TArray<UPrimitiveComponent*> CollisionComponents = CurrWaterBodyComponent->GetCollisionComponents();
+				for (UPrimitiveComponent* CurrCollisionComponent : CollisionComponents)
+				{
+					USplineMeshComponent* CurrSplineComponent = StaticCast<USplineMeshComponent*>(CurrCollisionComponent);
+					TObjectPtr<UBodySetup> CurrBodySetup = CurrSplineComponent->BodySetup;
+
+					const FTransform CurrMeshTransform = CurrCollisionComponent->GetComponentTransform();
+
+					// make sure the collision convex hull vertices are clamped to the min/max water height
+					for (FKConvexElem& ConvexElem : CurrBodySetup->AggGeom.ConvexElems)
+					{					
+						const TArray<FVector>& VertexData = ConvexElem.VertexData;		
+
+						// see if the current point is inside the convex projected to the xy plane
+						const FBox CurrBox = ConvexElem.CalcAABB(CurrMeshTransform, FVector(1, 1, 1));										
+
+						if (CurrBox.IsInsideXY(FBox(WorldPos, WorldPos)))
+						{						
+							float* TmpMaxHeight = ConvexToMaxHeight.Find(&ConvexElem);
+							if (TmpMaxHeight == nullptr)
+							{							
+								ConvexToMaxHeight.Emplace(&ConvexElem, WorldPos.Z);
+							}
+							else
+							{
+								*TmpMaxHeight = FMath::Max(*TmpMaxHeight, WorldPos.Z);
+							}
+						}
+					}
+				}
+			}
+		}
+	}}
+
 	// set the sim texture on each water body that is in the simulated river.  
 	for (TObjectPtr<AWaterBody > CurrWaterBody : AllWaterBodies)
 	{
 		TObjectPtr<UWaterBodyComponent> CurrWaterBodyComponent = CurrWaterBody->GetWaterBodyComponent();
 				
 		CurrWaterBodyComponent->SetBakedShallowWaterSimulation(BakedSim);
+
+		// grow bounds in z to include the tallest height
+		TArray<UPrimitiveComponent*> CollisionComponents = CurrWaterBodyComponent->GetCollisionComponents();
+
+		// Make sure that the collision objects includes the maximum height of the baked water sim otherwise
+		// we will miss collisions.
+		for (UPrimitiveComponent* CurrCollisionComponent : CollisionComponents)
+		{
+			USplineMeshComponent* CurrSplineComponent = StaticCast<USplineMeshComponent*>(CurrCollisionComponent);
+			
+			TObjectPtr<UBodySetup> CurrBodySetup = CurrSplineComponent->BodySetup;
+
+			FTransform CurrMeshTransform = CurrCollisionComponent->GetComponentTransform();
+
+			// make sure the collision convex hull vertices are clamped to the min/max water height
+			for (FKConvexElem& ConvexElem : CurrBodySetup->AggGeom.ConvexElems)
+			{			
+				// see if the current point is inside the convex projected to the xy plane
+				FBox CurrBox = ConvexElem.CalcAABB(CurrMeshTransform, FVector(1, 1, 1));				
+
+				TArray<FVector>& VertexData = ConvexElem.VertexData;
+				
+				// 
+				// We know based on the way each convex elem is created which indices correspond with
+				// top and bottom of the box.  From SplineMeshComponent.cpp:1075
+				// ConvexElem.VertexData.Add(ElementTM.TransformPosition(Radii * FVector(-1, -1, -1)));
+				// ConvexElem.VertexData.Add(ElementTM.TransformPosition(Radii * FVector(-1, -1, 1)));
+				// ConvexElem.VertexData.Add(ElementTM.TransformPosition(Radii * FVector(-1, 1, -1)));
+				// ConvexElem.VertexData.Add(ElementTM.TransformPosition(Radii * FVector(-1, 1, 1)));
+				// ConvexElem.VertexData.Add(ElementTM.TransformPosition(Radii * FVector(1, -1, -1)));
+				// ConvexElem.VertexData.Add(ElementTM.TransformPosition(Radii * FVector(1, -1, 1)));
+				// ConvexElem.VertexData.Add(ElementTM.TransformPosition(Radii * FVector(1, 1, -1)));
+				// ConvexElem.VertexData.Add(ElementTM.TransformPosition(Radii * FVector(1, 1, 1)));
+				
+				const float WorldMaxZForConvex = ConvexToMaxHeight[&ConvexElem];				
+
+				// for each vertex in the convex hull, set the Z to the maximum baked water sim Z height for the convex
+				int32 Idx = 0;
+				for (FVector& Vertex : VertexData)
+				{
+					FVector VWorld = CurrMeshTransform.TransformPosition(Vertex);					
+					
+					// Only set on top vertices of the convex hull with +Z to push it up
+					if (Idx == 1 || Idx == 3 || Idx == 5 || Idx == 7)
+					{
+						VWorld.Z = WorldMaxZForConvex;
+					}
+
+					const FVector VLocal = CurrMeshTransform.InverseTransformPosition(VWorld);
+					Vertex.X = VLocal.X;
+					Vertex.Y = VLocal.Y;
+					Vertex.Z = VLocal.Z;
+
+					Idx++;
+				}
+			}			
+
+			CurrSplineComponent->PostEditChange();
+		}
+
+		CurrWaterBodyComponent->PostEditChange();
 	}
 }
 
