- fix crash when baking due to generation of new water body collision bounds
- fix potential crash with baking button
- fix crash when opening maps due to null material instance

diff --git a/Engine/Plugins/Experimental/WaterAdvanced/Source/WaterAdvanced/Private/ShallowWaterRiverActor.cpp b/Engine/Plugins/Experimental/WaterAdvanced/Source/WaterAdvanced/Private/ShallowWaterRiverActor.cpp
--- a/Engine/Plugins/Experimental/WaterAdvanced/Source/WaterAdvanced/Private/ShallowWaterRiverActor.cpp
+++ b/Engine/Plugins/Experimental/WaterAdvanced/Source/WaterAdvanced/Private/ShallowWaterRiverActor.cpp
@@ -539,26 +539,27 @@ void UShallowWaterRiverComponent::Bake()
 				// ConvexElem.VertexData.Add(ElementTM.TransformPosition(Radii * FVector(1, 1, -1)));
 				// ConvexElem.VertexData.Add(ElementTM.TransformPosition(Radii * FVector(1, 1, 1)));
 				
-				const float WorldMaxZForConvex = ConvexToMaxHeight[&ConvexElem];				
-
-				// for each vertex in the convex hull, set the Z to the maximum baked water sim Z height for the convex
-				int32 Idx = 0;
-				for (FVector& Vertex : VertexData)
+				if (const float* WorldMaxZForConvex = ConvexToMaxHeight.Find(&ConvexElem))
 				{
-					FVector VWorld = CurrMeshTransform.TransformPosition(Vertex);					
-					
-					// Only set on top vertices of the convex hull with +Z to push it up
-					if (Idx == 1 || Idx == 3 || Idx == 5 || Idx == 7)
+					// for each vertex in the convex hull, set the Z to the maximum baked water sim Z height for the convex
+					int32 Idx = 0;
+					for (FVector& Vertex : VertexData)
 					{
-						VWorld.Z = WorldMaxZForConvex;
-					}
+						FVector VWorld = CurrMeshTransform.TransformPosition(Vertex);					
+					
+						// Only set on top vertices of the convex hull with +Z to push it up
+						if (Idx == 1 || Idx == 3 || Idx == 5 || Idx == 7)
+						{
+							VWorld.Z = *WorldMaxZForConvex;
+						}
 
-					const FVector VLocal = CurrMeshTransform.InverseTransformPosition(VWorld);
-					Vertex.X = VLocal.X;
-					Vertex.Y = VLocal.Y;
-					Vertex.Z = VLocal.Z;
+						const FVector VLocal = CurrMeshTransform.InverseTransformPosition(VWorld);
+						Vertex.X = VLocal.X;
+						Vertex.Y = VLocal.Y;
+						Vertex.Z = VLocal.Z;
 
-					Idx++;
+						Idx++;
+					}
 				}
 			}			
 
@@ -672,39 +673,55 @@ void UShallowWaterRiverComponent::UpdateRenderState()
 			UMaterialInstanceDynamic* WaterInfoMID = CurrWaterBodyComponent->GetWaterInfoMaterialInstance();
 			if (RenderState == EShallowWaterRenderState::WaterComponentWithBakedSim)
 			{				
-				// override materials on water bodies							
-				WaterMID->SetTextureParameterValue("BakedWaterSimTex", BakedWaterSurfaceTexture);
-				WaterMID->SetVectorParameterValue("BakedWaterSimLocation", SystemPos);
-				WaterMID->SetVectorParameterValue("BakedWaterSimSize", FVector(WorldGridSize.X, WorldGridSize.Y, 1));
-
-				WaterMID->SetTextureParameterValue("BakedWaterSimNormalTex", BakedWaterSurfaceNormalRT);
-				WaterMID->SetTextureParameterValue("NormalDetailTex", NormalDetailRT);
-				WaterMID->SetTextureParameterValue("FoamTex", FoamRT);
-
-				// 
+				// override materials on water bodies
 				
-				float Dx;
-				if (WorldGridSize.X > WorldGridSize.Y)
+				if (WaterMID)
 				{
-					Dx = WorldGridSize.X / ResolutionMaxAxis;
+					WaterMID->SetTextureParameterValue("BakedWaterSimTex", BakedWaterSurfaceTexture);
+					WaterMID->SetVectorParameterValue("BakedWaterSimLocation", SystemPos);
+					WaterMID->SetVectorParameterValue("BakedWaterSimSize", FVector(WorldGridSize.X, WorldGridSize.Y, 1));
+
+					WaterMID->SetTextureParameterValue("BakedWaterSimNormalTex", BakedWaterSurfaceNormalRT);
+					WaterMID->SetTextureParameterValue("NormalDetailTex", NormalDetailRT);
+					WaterMID->SetTextureParameterValue("FoamTex", FoamRT);					
+				
+					float Dx;
+					if (WorldGridSize.X > WorldGridSize.Y)
+					{
+						Dx = WorldGridSize.X / ResolutionMaxAxis;
+					}
+					else					
+					{
+						Dx = WorldGridSize.Y / ResolutionMaxAxis;
+					}				
+					WaterMID->SetScalarParameterValue("BakedWaterSimDx", Dx);
+
+					WaterMID->SetScalarParameterValue("UseBakedSimHack", RenderState == EShallowWaterRenderState::WaterComponentWithBakedSim ? 1 : 0);
 				}
-				else					
+				else
 				{
-					Dx = WorldGridSize.Y / ResolutionMaxAxis;
-				}				
-				WaterMID->SetScalarParameterValue("BakedWaterSimDx", Dx);
-
+					UE_LOG(LogShallowWater, Warning, TEXT("UShallowWaterRiverComponent::UpdateRenderState() - Water Component MID is null"));
+					return;
+				}
 
-				WaterInfoMID->SetTextureParameterValue("BakedWaterSimTex", BakedWaterSurfaceTexture);
-				WaterInfoMID->SetVectorParameterValue("BakedWaterSimLocation", SystemPos);
-				WaterInfoMID->SetVectorParameterValue("BakedWaterSimSize", FVector(WorldGridSize.X, WorldGridSize.Y, 1));
+				if (WaterInfoMID)
+				{
+					WaterInfoMID->SetTextureParameterValue("BakedWaterSimTex", BakedWaterSurfaceTexture);
+					WaterInfoMID->SetVectorParameterValue("BakedWaterSimLocation", SystemPos);
+					WaterInfoMID->SetVectorParameterValue("BakedWaterSimSize", FVector(WorldGridSize.X, WorldGridSize.Y, 1));
+				}
+				else
+				{
+					UE_LOG(LogShallowWater, Warning, TEXT("UShallowWaterRiverComponent::UpdateRenderState() - Water Component Water Info MID is null"));
+					return;
+				}
 			}
 
 			CurrWaterBodyComponent->SetUseBakedSimulationForQueriesAndPhysics(
 				RenderState == EShallowWaterRenderState::WaterComponentWithBakedSim || RenderState == EShallowWaterRenderState::BakedSim);
 
 			/*
-			#todo(dmp): I'd prefer if we could set an editor time only static switch to control using baked sims in the material or not
+			// #todo(dmp): I'd prefer if we could set an editor time only static switch to control using baked sims in the material or not
 			TArray<FMaterialParameterInfo> OutMaterialParameterInfos;
 			TArray<FGuid> Guids;
 			WaterMID->GetAllStaticSwitchParameterInfo(OutMaterialParameterInfos, Guids);
@@ -716,9 +733,7 @@ void UShallowWaterRiverComponent::UpdateRenderState()
 					WaterMID->SetStaticSwitchParameterValueEditorOnly(MaterialParameterInfo, RenderState == EShallowWaterRenderState::WaterComponentWithBakedSim);
 				}
 			}
-			*/
-
-			WaterMID->SetScalarParameterValue("UseBakedSimHack", RenderState == EShallowWaterRenderState::WaterComponentWithBakedSim ? 1 : 0);
+			*/			
 		}
 	}
 }
diff --git a/Engine/Plugins/Experimental/WaterAdvanced/Source/WaterAdvanced/Private/ShallowWaterRiverDetails.cpp b/Engine/Plugins/Experimental/WaterAdvanced/Source/WaterAdvanced/Private/ShallowWaterRiverDetails.cpp
--- a/Engine/Plugins/Experimental/WaterAdvanced/Source/WaterAdvanced/Private/ShallowWaterRiverDetails.cpp
+++ b/Engine/Plugins/Experimental/WaterAdvanced/Source/WaterAdvanced/Private/ShallowWaterRiverDetails.cpp
@@ -132,13 +132,15 @@ FReply FShallowWaterRiverDetails::OnResetSelectedSystem()
 			{
 				for (UActorComponent* AC : Actor->GetComponents())
 				{
-					UShallowWaterRiverComponent* ShallowWaterComponent = Cast<UShallowWaterRiverComponent>(AC);
-					if (ShallowWaterComponent)
+					if (AC)
 					{
-						ShallowWaterComponent->RenderState = EShallowWaterRenderState::LiveSim;
-						ShallowWaterComponent->UpdateRenderState();
-						ShallowWaterComponent->Rebuild();
-						ShallowWaterComponent->ReregisterComponent();
+						if (UShallowWaterRiverComponent* ShallowWaterComponent = Cast<UShallowWaterRiverComponent>(AC))
+						{
+							ShallowWaterComponent->RenderState = EShallowWaterRenderState::LiveSim;
+							ShallowWaterComponent->UpdateRenderState();
+							ShallowWaterComponent->Rebuild();
+							ShallowWaterComponent->ReregisterComponent();
+						}
 					}
 				}
 			}
@@ -171,11 +173,13 @@ FReply FShallowWaterRiverDetails::OnBakeSelectedSystem()
 			{
 				for (UActorComponent* AC : Actor->GetComponents())
 				{
-					UShallowWaterRiverComponent* ShallowWaterComponent = Cast<UShallowWaterRiverComponent>(AC);
-					if (ShallowWaterComponent)
+					if (AC)
 					{
-						ShallowWaterComponent->Bake();
-						break;
+						if (UShallowWaterRiverComponent* ShallowWaterComponent = Cast<UShallowWaterRiverComponent>(AC))
+						{
+							ShallowWaterComponent->Bake();
+							break;
+						}
 					}
 				}
 			}
