check for null

diff --git a/Engine/Plugins/Experimental/WaterAdvanced/Source/WaterAdvanced/Private/ShallowWaterRiverActor.cpp b/Engine/Plugins/Experimental/WaterAdvanced/Source/WaterAdvanced/Private/ShallowWaterRiverActor.cpp
--- a/Engine/Plugins/Experimental/WaterAdvanced/Source/WaterAdvanced/Private/ShallowWaterRiverActor.cpp
+++ b/Engine/Plugins/Experimental/WaterAdvanced/Source/WaterAdvanced/Private/ShallowWaterRiverActor.cpp
@@ -663,6 +663,12 @@ void UShallowWaterRiverComponent::UpdateRenderState()
 
 	for (AWaterBody* CurrWaterBody : AllWaterBodies)
 	{
+		if (!CurrWaterBody)
+		{
+			UE_LOG(LogShallowWater, Warning, TEXT("UShallowWaterRiverComponent::UpdateRenderState() - Water Body Actor is null"));
+			continue;
+		}
+
 		TObjectPtr<UWaterBodyComponent> CurrWaterBodyComponent = CurrWaterBody->GetWaterBodyComponent();
 
 		if (CurrWaterBodyComponent != nullptr)
