tagging for bottom contour colliders
bottom colliders works with LevelInstance actors

#virtualized

diff --git a/Engine/Plugins/Experimental/WaterAdvanced/Source/WaterAdvanced/Private/ShallowWaterRiverActor.cpp b/Engine/Plugins/Experimental/WaterAdvanced/Source/WaterAdvanced/Private/ShallowWaterRiverActor.cpp
--- a/Engine/Plugins/Experimental/WaterAdvanced/Source/WaterAdvanced/Private/ShallowWaterRiverActor.cpp
+++ b/Engine/Plugins/Experimental/WaterAdvanced/Source/WaterAdvanced/Private/ShallowWaterRiverActor.cpp
@@ -17,6 +17,7 @@
 #include "Components/SplineMeshComponent.h"
 #include "PhysicsEngine/ConvexElem.h"
 #include "PhysicsEngine/BodySetup.h"
+#include "Engine/OverlapResult.h"
 
 #include "TextureResource.h"
 #include "ShallowWaterCommon.h"
@@ -25,6 +26,8 @@
 #include "Math/Float16Color.h"
 #include "Kismet/GameplayStatics.h"
 #include "Landscape.h"
+#include "LevelInstance/LevelInstanceActor.h"
+#include "LevelInstance/LevelInstanceSubsystem.h"
 
 UShallowWaterRiverComponent::UShallowWaterRiverComponent(const FObjectInitializer& Initializer)
 	: Super(Initializer)
@@ -262,7 +265,64 @@ void UShallowWaterRiverComponent::Rebuild()
 			for (TObjectPtr<AActor> CurrActor : BottomContourActors)
 			{
 				AActor* CurrActorRawPtr = CurrActor.Get();
-				BottomContourActorsRawPtr.Add(CurrActorRawPtr);
+
+				// if we have a level instance, break it up and add each actor
+				if (ALevelInstance* LevelInstancePtr = Cast<ALevelInstance>(CurrActorRawPtr))
+				{
+					const ULevelInstanceSubsystem* LevelInstanceSubsystem = GetWorld()->GetSubsystem<ULevelInstanceSubsystem>();
+
+					LevelInstanceSubsystem->ForEachActorInLevelInstance(LevelInstancePtr, [&](AActor* SubActor)
+						{
+							BottomContourActorsRawPtr.Add(SubActor);
+							return true;
+						});
+				}
+				else
+				{
+					BottomContourActorsRawPtr.Add(CurrActorRawPtr);
+				}				
+			}
+
+			// if we have a tag set
+			// do an overlap test
+			//  filter by tag and add to the bottomcontour actors list
+			//  if a level instance is tagged, loop over the contained actors
+
+			if (!BottomContourTags.IsEmpty())
+			{								
+				FCollisionQueryParams Params(SCENE_QUERY_STAT(ShallowWaterRiverActorQuery), false);
+
+				TArray<FOverlapResult> Overlaps;								
+				GetWorld()->OverlapMultiByChannel(Overlaps, SystemPos, FQuat::Identity, ECollisionChannel::ECC_WorldStatic,
+					FCollisionShape::MakeBox(0.5f * FVector(WorldGridSize.X, WorldGridSize.Y, 100000)), Params);
+
+				for (const FOverlapResult& OverlapResult : Overlaps)
+				{
+					if (UPrimitiveComponent* PrimitiveComponent = OverlapResult.GetComponent())
+					{
+						if (AActor* ComponentActor = PrimitiveComponent->GetOwner())
+						{
+							if (BottomContourTags.ContainsByPredicate([&](const FName& Tag) { return Tag == NAME_None || ComponentActor->Tags.Contains(Tag); }))
+							{
+								// if we have a level instance, break it up and add each actor
+								if (ALevelInstance* LevelInstancePtr = Cast<ALevelInstance>(ComponentActor))
+								{
+									const ULevelInstanceSubsystem* LevelInstanceSubsystem = GetWorld()->GetSubsystem<ULevelInstanceSubsystem>();
+
+									LevelInstanceSubsystem->ForEachActorInLevelInstance(LevelInstancePtr, [&](AActor* SubActor)
+										{
+											BottomContourActorsRawPtr.Add(SubActor);
+											return true;
+										});
+								}
+								else
+								{
+									BottomContourActorsRawPtr.Add(ComponentActor);
+								}
+							}
+						}
+					}
+				}
 			}
 
 			FName DIName = "User.BottomCapture";
diff --git a/Engine/Plugins/Experimental/WaterAdvanced/Source/WaterAdvanced/Public/ShallowWaterRiverActor.h b/Engine/Plugins/Experimental/WaterAdvanced/Source/WaterAdvanced/Public/ShallowWaterRiverActor.h
--- a/Engine/Plugins/Experimental/WaterAdvanced/Source/WaterAdvanced/Public/ShallowWaterRiverActor.h
+++ b/Engine/Plugins/Experimental/WaterAdvanced/Source/WaterAdvanced/Public/ShallowWaterRiverActor.h
@@ -78,6 +78,9 @@ public:
 	UPROPERTY(EditAnywhere, Category = "Collisions")
 	TArray<TObjectPtr<AActor>> BottomContourActors;
 
+	UPROPERTY(EditAnywhere, Category = "Collisions")
+	TArray<FName> BottomContourTags;
+
 	UPROPERTY(EditAnywhere, Category = "Collisions")
 	float BottomContourCaptureOffset = 1000.f;
 
