change ensures to warnings

diff --git a/Engine/Plugins/Experimental/WaterAdvanced/Source/WaterAdvanced/Private/ShallowWaterRiverActor.cpp b/Engine/Plugins/Experimental/WaterAdvanced/Source/WaterAdvanced/Private/ShallowWaterRiverActor.cpp
--- a/Engine/Plugins/Experimental/WaterAdvanced/Source/WaterAdvanced/Private/ShallowWaterRiverActor.cpp
+++ b/Engine/Plugins/Experimental/WaterAdvanced/Source/WaterAdvanced/Private/ShallowWaterRiverActor.cpp
@@ -687,7 +687,7 @@ void UShallowWaterRiverComponent::OnWaterInfoTextureArrayCreated(const UTextureR
 {	
 	if (InWaterInfoTexture == nullptr)
 	{
-		ensureMsgf(false, TEXT("UShallowWaterRiverComponent::OnWaterInfoTextureCreated was called with NULL WaterInfoTexture"));
+		UE_LOG(LogShallowWater, Warning, TEXT("UShallowWaterRiverComponent::OnWaterInfoTextureCreated was called with NULL WaterInfoTexture"));
 		return;
 	}
 	
@@ -697,7 +697,7 @@ void UShallowWaterRiverComponent::OnWaterInfoTextureArrayCreated(const UTextureR
 		UTexture* WITTextureArray = Cast<UTexture>(const_cast<UTextureRenderTarget2DArray*>(WaterInfoTexture.Get()));
 		if (WITTextureArray == nullptr)
 		{
-			ensureMsgf(false, TEXT("UShallowWaterRiverComponent::OnWaterInfoTextureCreated was called with Water Info Texture that isn't valid"));
+			UE_LOG(LogShallowWater, Warning, TEXT("UShallowWaterRiverComponent::OnWaterInfoTextureCreated was called with Water Info Texture that isn't valid"));
 			return;
 		}
 
@@ -705,7 +705,7 @@ void UShallowWaterRiverComponent::OnWaterInfoTextureArrayCreated(const UTextureR
 	}
 	else
 	{
-		ensureMsgf(false, TEXT("UShallowWaterRiverComponent::OnWaterInfoTextureCreated was called with NULL ShallowWaterNiagaraSimulation"));
+		UE_LOG(LogShallowWater, Warning, TEXT("UShallowWaterRiverComponent::OnWaterInfoTextureCreated was called with NULL ShallowWaterNiagaraSimulation"));
 		return;
 	}
 }
