better error reporting

diff --git a/Engine/Plugins/Experimental/WaterAdvanced/Source/WaterAdvanced/Private/ShallowWaterRiverActor.cpp b/Engine/Plugins/Experimental/WaterAdvanced/Source/WaterAdvanced/Private/ShallowWaterRiverActor.cpp
--- a/Engine/Plugins/Experimental/WaterAdvanced/Source/WaterAdvanced/Private/ShallowWaterRiverActor.cpp
+++ b/Engine/Plugins/Experimental/WaterAdvanced/Source/WaterAdvanced/Private/ShallowWaterRiverActor.cpp
@@ -86,7 +86,19 @@ void UShallowWaterRiverComponent::TickComponent(float DeltaTime, ELevelTick Tick
 	}
 	else if (bIsInitialized)
 	{
-		RiverSimSystem->Activate();	
+		if (RiverSimSystem)
+		{
+			RiverSimSystem->Activate();	
+		}
+		else
+		{
+			UE_LOG(LogShallowWater, Warning, TEXT("UShallowWaterRiverComponent::TickComponent() - null Niagara sim when trying to activate. Please reset."));
+		}
+	}
+	else
+	{
+		// System is in a bad state
+		// UE_LOG(LogShallowWater, Warning, TEXT("UShallowWaterRiverComponent::TickComponent() - null Niagara sim when trying to activate. Please reset."));
 	}
 #endif
 }
@@ -122,6 +134,9 @@ void UShallowWaterRiverComponent::PostEditChangeProperty(FPropertyChangedEvent&
 	}
 	else
 	{
+		bIsInitialized = false;
+		bTickInitialize = false;
+
 		Rebuild();
 	}
 
@@ -133,6 +148,8 @@ void UShallowWaterRiverComponent::PostEditChangeProperty(FPropertyChangedEvent&
 
 void UShallowWaterRiverComponent::Rebuild()
 {	
+	bIsInitialized = false;
+
 	if (NiagaraRiverSimulation == nullptr)
 	{
 		NiagaraRiverSimulation = LoadObject<UNiagaraSystem>(nullptr, TEXT("/WaterAdvanced/Niagara/Systems/Grid2D_SW_River.Grid2D_SW_River"));
@@ -145,9 +162,28 @@ void UShallowWaterRiverComponent::Rebuild()
 		RiverSimSystem = nullptr;
 	}
 	
+	if (ResolutionMaxAxis <= 0)
+	{
+		UE_LOG(LogShallowWater, Warning, TEXT("UShallowWaterRiverComponent::Rebuild() - resolution must be greater than 0"));
+		return;
+	}
+
+	if (NumSteps <= 0)
+	{
+		UE_LOG(LogShallowWater, Warning, TEXT("UShallowWaterRiverComponent::Rebuild() - num steps must be greater than 0"));
+		return;
+	}
+
+	if (SimSpeed <= 1e-8)
+	{
+		UE_LOG(LogShallowWater, Warning, TEXT("UShallowWaterRiverComponent::Rebuild() - speed must be greater than zero"));
+		return;
+	}
+
 	if (NiagaraRiverSimulation == nullptr)
 	{
 		UE_LOG(LogShallowWater, Warning, TEXT("UShallowWaterRiverComponent::Rebuild() - null Niagara system asset"));
+		return;
 	}
 
 	AllWaterBodies.Empty();
@@ -416,6 +452,13 @@ void UShallowWaterRiverComponent::Rebuild()
 	RiverSimSystem->Activate();
 	
 	WorldGridSize = 2.0f * FVector2D(CombinedBounds.BoxExtent.X, CombinedBounds.BoxExtent.Y);
+
+	if (WorldGridSize.Length() < 1e-8)
+	{
+		UE_LOG(LogShallowWater, Warning, TEXT("UShallowWaterRiverComponent::Bake() - Simulation grid has (0,0) size."));
+		return;
+	}
+
 	RiverSimSystem->SetVariableVec2(FName("WorldGridSize"), WorldGridSize);
 	RiverSimSystem->SetVariableInt(FName("ResolutionMaxAxis"), ResolutionMaxAxis);
 
@@ -468,10 +511,9 @@ void UShallowWaterRiverComponent::Bake()
 {
 	EObjectFlags TextureObjectFlags = EObjectFlags::RF_Public;
 
-	if (!BakedWaterSurfaceRT)
+	if (!RiverSimSystem || !BakedWaterSurfaceRT)
 	{
 		UE_LOG(LogShallowWater, Warning, TEXT("UShallowWaterRiverComponent::Bake() - No simulation to bake"));
-
 		return;
 	}
 
@@ -759,7 +801,7 @@ void UShallowWaterRiverComponent::UpdateRenderState()
 					else					
 					{
 						Dx = WorldGridSize.Y / ResolutionMaxAxis;
-					}				
+					}
 					WaterMID->SetScalarParameterValue("BakedWaterSimDx", Dx);
 
 					WaterMID->SetScalarParameterValue("UseBakedSimHack", RenderState == EShallowWaterRenderState::WaterComponentWithBakedSim ? 1 : 0);
