option to debug render bottom contour

#virtualized

diff --git a/Engine/Plugins/Experimental/WaterAdvanced/Source/WaterAdvanced/Private/ShallowWaterRiverActor.cpp b/Engine/Plugins/Experimental/WaterAdvanced/Source/WaterAdvanced/Private/ShallowWaterRiverActor.cpp
--- a/Engine/Plugins/Experimental/WaterAdvanced/Source/WaterAdvanced/Private/ShallowWaterRiverActor.cpp
+++ b/Engine/Plugins/Experimental/WaterAdvanced/Source/WaterAdvanced/Private/ShallowWaterRiverActor.cpp
@@ -140,10 +140,8 @@ void UShallowWaterRiverComponent::PostEditChangeProperty(FPropertyChangedEvent&
 		Rebuild();
 	}
 
-	if (PropertyName == GET_MEMBER_NAME_CHECKED(UShallowWaterRiverComponent, RenderState))
-	{
-		UpdateRenderState();
-	}
+
+	UpdateRenderState();
 }
 
 void UShallowWaterRiverComponent::Rebuild()
@@ -761,6 +759,8 @@ void UShallowWaterRiverComponent::UpdateRenderState()
 	if (RiverSimSystem != nullptr)
 	{		
 		RiverSimSystem->SetVariableBool(FName("RenderWaterSurface"), !RenderWaterBody);
+		RiverSimSystem->SetVariableBool(FName("DebugRender"), RenderState == EShallowWaterRenderState::DebugRenderBottomContour);
+		
 	}
 
 	for (AWaterBody* CurrWaterBody : AllWaterBodies)
diff --git a/Engine/Plugins/Experimental/WaterAdvanced/Source/WaterAdvanced/Public/ShallowWaterRiverActor.h b/Engine/Plugins/Experimental/WaterAdvanced/Source/WaterAdvanced/Public/ShallowWaterRiverActor.h
--- a/Engine/Plugins/Experimental/WaterAdvanced/Source/WaterAdvanced/Public/ShallowWaterRiverActor.h
+++ b/Engine/Plugins/Experimental/WaterAdvanced/Source/WaterAdvanced/Public/ShallowWaterRiverActor.h
@@ -29,7 +29,8 @@ enum EShallowWaterRenderState : int
 	WaterComponent,
 	WaterComponentWithBakedSim,
 	LiveSim,
-	BakedSim
+	BakedSim,
+	DebugRenderBottomContour,
 };
 
 
@@ -65,7 +66,7 @@ public:
 	
 	UPROPERTY(EditAnywhere, Category = "Rendering", meta = (DisplayName = "Render State"))
 	TEnumAsByte<EShallowWaterRenderState> RenderState = EShallowWaterRenderState::WaterComponent;
-
+	
 	UPROPERTY(EditAnywhere, Category = "Shallow Water")
 	TObjectPtr<UTexture2D> BakedWaterSurfaceTexture;
 	
@@ -73,7 +74,7 @@ public:
 	//TObjectPtr<UTexture2D> SignedDistanceToSplineTexture;
 
 	UPROPERTY(EditAnywhere, Category = "Collisions")
-	bool bUseCapture = false;
+	bool bUseCapture = true;
 
 	UPROPERTY(EditAnywhere, Category = "Collisions")
 	TArray<TObjectPtr<AActor>> BottomContourActors;
