option to debug foam emission

#virtualized

diff --git a/Engine/Plugins/Experimental/WaterAdvanced/Source/WaterAdvanced/Private/ShallowWaterRiverActor.cpp b/Engine/Plugins/Experimental/WaterAdvanced/Source/WaterAdvanced/Private/ShallowWaterRiverActor.cpp
--- a/Engine/Plugins/Experimental/WaterAdvanced/Source/WaterAdvanced/Private/ShallowWaterRiverActor.cpp
+++ b/Engine/Plugins/Experimental/WaterAdvanced/Source/WaterAdvanced/Private/ShallowWaterRiverActor.cpp
@@ -759,7 +759,8 @@ void UShallowWaterRiverComponent::UpdateRenderState()
 	if (RiverSimSystem != nullptr)
 	{		
 		RiverSimSystem->SetVariableBool(FName("RenderWaterSurface"), !RenderWaterBody);
-		RiverSimSystem->SetVariableBool(FName("DebugRender"), RenderState == EShallowWaterRenderState::DebugRenderBottomContour);
+		RiverSimSystem->SetVariableBool(FName("DebugRenderBottomContour"), RenderState == EShallowWaterRenderState::DebugRenderBottomContour);
+		RiverSimSystem->SetVariableBool(FName("DebugRenderFoam"), RenderState == EShallowWaterRenderState::DebugRenderFoam);
 		
 	}
 
diff --git a/Engine/Plugins/Experimental/WaterAdvanced/Source/WaterAdvanced/Public/ShallowWaterRiverActor.h b/Engine/Plugins/Experimental/WaterAdvanced/Source/WaterAdvanced/Public/ShallowWaterRiverActor.h
--- a/Engine/Plugins/Experimental/WaterAdvanced/Source/WaterAdvanced/Public/ShallowWaterRiverActor.h
+++ b/Engine/Plugins/Experimental/WaterAdvanced/Source/WaterAdvanced/Public/ShallowWaterRiverActor.h
@@ -31,6 +31,7 @@ enum EShallowWaterRenderState : int
 	LiveSim,
 	BakedSim,
 	DebugRenderBottomContour,
+	DebugRenderFoam,
 };
 
 
