expose collision dilation
expose render time smoothing options
water body and water body wtih baked sim rendering modes work as expected

#virtualized

diff --git a/Engine/Plugins/Experimental/WaterAdvanced/Source/WaterAdvanced/Private/ShallowWaterRiverActor.cpp b/Engine/Plugins/Experimental/WaterAdvanced/Source/WaterAdvanced/Private/ShallowWaterRiverActor.cpp
--- a/Engine/Plugins/Experimental/WaterAdvanced/Source/WaterAdvanced/Private/ShallowWaterRiverActor.cpp
+++ b/Engine/Plugins/Experimental/WaterAdvanced/Source/WaterAdvanced/Private/ShallowWaterRiverActor.cpp
@@ -56,7 +56,7 @@ UShallowWaterRiverComponent::UShallowWaterRiverComponent(const FObjectInitialize
 
 void UShallowWaterRiverComponent::PostLoad()
 {
-	Super::PostLoad();
+	Super::PostLoad();	
 
 	if (RenderState == EShallowWaterRenderState::LiveSim || RiverSimSystem == nullptr)
 	{
@@ -482,6 +482,11 @@ void UShallowWaterRiverComponent::Rebuild()
 	RiverSimSystem->SetVariableTextureRenderTarget(FName("SimGridRT"), BakedWaterSurfaceRT);
 	RiverSimSystem->SetVariableBool(FName("ReadCachedSim"), RenderState == EShallowWaterRenderState::BakedSim || RenderState == EShallowWaterRenderState::WaterComponentWithBakedSim);
 
+	RiverSimSystem->SetVariableFloat(FName("BottomContourCollisionDilation"), BottomContourCollisionDilation);
+	
+	RiverSimSystem->SetVariableInt(FName("ExtrapolationHalfWidth"), SmoothingWidth);
+	RiverSimSystem->SetVariableFloat(FName("SmoothingHeightCutoff"), SmoothingCutoff);
+
 	if (BakedWaterSurfaceTexture != nullptr)
 	{
 		RiverSimSystem->SetVariableTexture(FName("BakedSimTexture"), BakedWaterSurfaceTexture);
@@ -754,6 +759,11 @@ void UShallowWaterRiverComponent::OnWaterInfoTextureArrayCreated(const UTextureR
 
 void UShallowWaterRiverComponent::UpdateRenderState()
 {
+	if (BakedSimMaterial == nullptr)
+	{
+		BakedSimMaterial = LoadObject<UMaterialInstance>(nullptr, TEXT("/WaterAdvanced/Niagara/Materials/M_SW_River_Inst.M_SW_River_Inst"));
+	}
+	
 	bool RenderWaterBody = RenderState == EShallowWaterRenderState::WaterComponent || RenderState == EShallowWaterRenderState::WaterComponentWithBakedSim;
 
 	if (RiverSimSystem != nullptr)
@@ -778,10 +788,18 @@ void UShallowWaterRiverComponent::UpdateRenderState()
 		{
 			CurrWaterBodyComponent->SetVisibility(RenderWaterBody);
 
-			UMaterialInstanceDynamic* WaterMID = CurrWaterBodyComponent->GetWaterMaterialInstance();
-			UMaterialInstanceDynamic* WaterInfoMID = CurrWaterBodyComponent->GetWaterInfoMaterialInstance();
 			if (RenderState == EShallowWaterRenderState::WaterComponentWithBakedSim)
-			{				
+			{			
+				if (OriginalWaterBodyMaterial == nullptr)
+				{
+					OriginalWaterBodyMaterial = CurrWaterBodyComponent->GetWaterMaterial();
+				}
+
+				CurrWaterBodyComponent->SetWaterMaterial(BakedSimMaterial);
+
+				UMaterialInstanceDynamic* WaterMID = CurrWaterBodyComponent->GetWaterMaterialInstance();
+				UMaterialInstanceDynamic* WaterInfoMID = CurrWaterBodyComponent->GetWaterInfoMaterialInstance();
+
 				// override materials on water bodies
 				
 				if (WaterMID)
@@ -825,6 +843,14 @@ void UShallowWaterRiverComponent::UpdateRenderState()
 					return;
 				}
 			}
+			else if (RenderState == EShallowWaterRenderState::WaterComponent)
+			{
+				if (OriginalWaterBodyMaterial != nullptr)
+				{
+					CurrWaterBodyComponent->SetWaterMaterial(OriginalWaterBodyMaterial.Get());
+					OriginalWaterBodyMaterial = nullptr;
+				}
+			}
 
 			CurrWaterBodyComponent->SetUseBakedSimulationForQueriesAndPhysics(
 				RenderState == EShallowWaterRenderState::WaterComponentWithBakedSim || RenderState == EShallowWaterRenderState::BakedSim);
diff --git a/Engine/Plugins/Experimental/WaterAdvanced/Source/WaterAdvanced/Public/ShallowWaterRiverActor.h b/Engine/Plugins/Experimental/WaterAdvanced/Source/WaterAdvanced/Public/ShallowWaterRiverActor.h
--- a/Engine/Plugins/Experimental/WaterAdvanced/Source/WaterAdvanced/Public/ShallowWaterRiverActor.h
+++ b/Engine/Plugins/Experimental/WaterAdvanced/Source/WaterAdvanced/Public/ShallowWaterRiverActor.h
@@ -68,6 +68,15 @@ public:
 	UPROPERTY(EditAnywhere, Category = "Rendering", meta = (DisplayName = "Render State"))
 	TEnumAsByte<EShallowWaterRenderState> RenderState = EShallowWaterRenderState::WaterComponent;
 	
+	UPROPERTY(EditAnywhere, Category = "Rendering", meta = (DisplayName = "Baked Sim Material"))
+	TObjectPtr <class UMaterialInstance> BakedSimMaterial;
+
+	UPROPERTY(EditAnywhere, Category = "Rendering", meta = (DisplayName = "Smoothing Width"))
+	int SmoothingWidth = 2;
+
+	UPROPERTY(EditAnywhere, Category = "Rendering", meta = (DisplayName = "Smoothing Cutoff"))
+	float SmoothingCutoff = 500.f;
+
 	UPROPERTY(EditAnywhere, Category = "Shallow Water")
 	TObjectPtr<UTexture2D> BakedWaterSurfaceTexture;
 	
@@ -86,6 +95,9 @@ public:
 	UPROPERTY(EditAnywhere, Category = "Collisions")
 	float BottomContourCaptureOffset = 1000.f;
 
+	UPROPERTY(EditAnywhere, Category = "Collisions")
+	float BottomContourCollisionDilation = 0.f;
+
 	virtual void PostLoad() override;
 
 	virtual void TickComponent(float DeltaTime, ELevelTick TickType, FActorComponentTickFunction* ThisTickFunction) override;
@@ -120,6 +132,9 @@ protected:
 	UPROPERTY(BlueprintReadOnly, VisibleDefaultsOnly, Category = "Shallow Water")
 	TObjectPtr<UBakedShallowWaterSimulationComponent> BakedSim;
 
+	UPROPERTY(BlueprintReadOnly, VisibleDefaultsOnly, Category = "Shallow Water")
+	TObjectPtr <UMaterialInterface> OriginalWaterBodyMaterial;
+
 	bool QueryWaterAtSplinePoint(TObjectPtr<AWaterBody> WaterBody, int SplinePoint, FVector& OutPos, FVector& OutTangent, float& OutWidth, float& OutDepth);
 	
 private:
