separate bottom contour capture for dilated actors
add Niagara template tags for water advanced systems
refactor source/sink so they can exist in the same stage
new base emitter for river sims

diff --git a/Engine/Plugins/Experimental/WaterAdvanced/Source/WaterAdvanced/Private/ShallowWaterRiverActor.cpp b/Engine/Plugins/Experimental/WaterAdvanced/Source/WaterAdvanced/Private/ShallowWaterRiverActor.cpp
--- a/Engine/Plugins/Experimental/WaterAdvanced/Source/WaterAdvanced/Private/ShallowWaterRiverActor.cpp
+++ b/Engine/Plugins/Experimental/WaterAdvanced/Source/WaterAdvanced/Private/ShallowWaterRiverActor.cpp
@@ -50,6 +50,12 @@ UShallowWaterRiverComponent::UShallowWaterRiverComponent(const FObjectInitialize
 	{
 		BottomContourActors.AddDefaulted();
 	}
+
+	// start with one empty element
+	if (DilatedBottomContourActors.IsEmpty())
+	{
+		DilatedBottomContourActors.AddDefaulted();
+	}
 #endif
 
 }
@@ -296,71 +302,12 @@ void UShallowWaterRiverComponent::Rebuild()
 		if (bUseCapture)
 		{
 			TArray<AActor*> BottomContourActorsRawPtr;
-			for (TObjectPtr<AActor> CurrActor : BottomContourActors)
-			{
-				AActor* CurrActorRawPtr = CurrActor.Get();
-
-				// if we have a level instance, break it up and add each actor
-				if (ALevelInstance* LevelInstancePtr = Cast<ALevelInstance>(CurrActorRawPtr))
-				{
-					const ULevelInstanceSubsystem* LevelInstanceSubsystem = GetWorld()->GetSubsystem<ULevelInstanceSubsystem>();
-
-					LevelInstanceSubsystem->ForEachActorInLevelInstance(LevelInstancePtr, [&](AActor* SubActor)
-						{
-							BottomContourActorsRawPtr.Add(SubActor);
-							return true;
-						});
-				}
-				else
-				{
-					BottomContourActorsRawPtr.Add(CurrActorRawPtr);
-				}				
-			}
-
-			// if we have a tag set
-			// do an overlap test
-			//  filter by tag and add to the bottomcontour actors list
-			//  if a level instance is tagged, loop over the contained actors
-
-			if (!BottomContourTags.IsEmpty())
-			{								
-				FCollisionQueryParams Params(SCENE_QUERY_STAT(ShallowWaterRiverActorQuery), false);
-
-				TArray<FOverlapResult> Overlaps;								
-				GetWorld()->OverlapMultiByChannel(Overlaps, SystemPos, FQuat::Identity, ECollisionChannel::ECC_WorldStatic,
-					FCollisionShape::MakeBox(0.5f * FVector(WorldGridSize.X, WorldGridSize.Y, 100000)), Params);
-
-				for (const FOverlapResult& OverlapResult : Overlaps)
-				{
-					if (UPrimitiveComponent* PrimitiveComponent = OverlapResult.GetComponent())
-					{
-						if (AActor* ComponentActor = PrimitiveComponent->GetOwner())
-						{
-							if (BottomContourTags.ContainsByPredicate([&](const FName& Tag) { return Tag == NAME_None || ComponentActor->Tags.Contains(Tag); }))
-							{
-								// if we have a level instance, break it up and add each actor
-								if (ALevelInstance* LevelInstancePtr = Cast<ALevelInstance>(ComponentActor))
-								{
-									const ULevelInstanceSubsystem* LevelInstanceSubsystem = GetWorld()->GetSubsystem<ULevelInstanceSubsystem>();
-
-									LevelInstanceSubsystem->ForEachActorInLevelInstance(LevelInstancePtr, [&](AActor* SubActor)
-										{
-											BottomContourActorsRawPtr.Add(SubActor);
-											return true;
-										});
-								}
-								else
-								{
-									BottomContourActorsRawPtr.Add(ComponentActor);
-								}
-							}
-						}
-					}
-				}
-			}
+			BottomContourActorsRawPtr.Add(nullptr);
 
+			AddActorsToRawArray(BottomContourActors, BottomContourActorsRawPtr);			
+			AddTaggedActorsToArray(BottomContourTags, BottomContourActorsRawPtr);			
+			
 			FName DIName = "User.BottomCapture";
-
 			UNiagaraFunctionLibrary::SetSceneCapture2DDataInterfaceManagedMode(RiverSimSystem, DIName,
 				ESceneCaptureSource::SCS_SceneDepth,
 				FIntPoint(ResolutionMaxAxis, ResolutionMaxAxis),
@@ -392,11 +339,50 @@ void UShallowWaterRiverComponent::Rebuild()
 				}
 			}
 			FBoxSphereBounds CombinedBottomContourBounds(BottomContourCombinedWorldBoundsBuilder);				
+			
+			// Dilated capture
+			TArray<AActor*> DilatedBottomContourActorsRawPtr;
+			DilatedBottomContourActorsRawPtr.Add(nullptr);
+			AddActorsToRawArray(DilatedBottomContourActors, DilatedBottomContourActorsRawPtr);
+			AddTaggedActorsToArray(DilatedBottomContourTags, DilatedBottomContourActorsRawPtr);
+
+			FName DilatedDIName = "User.DilatedBottomCapture";
+			UNiagaraFunctionLibrary::SetSceneCapture2DDataInterfaceManagedMode(RiverSimSystem, DilatedDIName,
+				ESceneCaptureSource::SCS_SceneDepth,
+				FIntPoint(ResolutionMaxAxis, ResolutionMaxAxis),
+				ETextureRenderTargetFormat::RTF_R16f,
+				ECameraProjectionMode::Orthographic,
+				90.0f,
+				FMath::Max(WorldGridSize.X, WorldGridSize.Y),
+				true,
+				false,
+				DilatedBottomContourActorsRawPtr);
 
 
+			// accumulate bounding box for river water bodies
+			FBoxSphereBounds::Builder DilatedBottomContourCombinedWorldBoundsBuilder;
+			for (AActor *BottomContourActor : DilatedBottomContourActorsRawPtr)
+			{
+				if (BottomContourActor != nullptr)
+				{
+					// accumulate bounds
+					FBoxSphereBounds WorldBounds;
+					BottomContourActor->GetActorBounds(false, WorldBounds.Origin, WorldBounds.BoxExtent);
+
+					DilatedBottomContourCombinedWorldBoundsBuilder += WorldBounds;
+				}
+				else
+				{
+					UE_LOG(LogShallowWater, Warning, TEXT("UShallowWaterRiverComponent::Rebuild() - skipping null bottom contour boundary actor found"));
+					continue;
+				}
+			}
+			FBoxSphereBounds DilatedCombinedBottomContourBounds(DilatedBottomContourCombinedWorldBoundsBuilder);
+			
+			// reinitialize and set variables on the system
 			RiverSimSystem->ReinitializeSystem();
-	
 			RiverSimSystem->SetVariableFloat(FName("CaptureOffset"), BottomContourCaptureOffset + CombinedBottomContourBounds.Origin.Z + CombinedBottomContourBounds.BoxExtent.Z);
+			RiverSimSystem->SetVariableFloat(FName("DilatedCaptureOffset"), BottomContourCaptureOffset + DilatedCombinedBottomContourBounds.Origin.Z + DilatedCombinedBottomContourBounds.BoxExtent.Z);
 		}
 		else
 		{
@@ -480,7 +466,7 @@ void UShallowWaterRiverComponent::Rebuild()
 	BakedWaterSurfaceRT = NewObject<UTextureRenderTarget2D>(this, NAME_None, RF_Transient);
 	BakedWaterSurfaceRT->InitAutoFormat(1, 1);
 	RiverSimSystem->SetVariableTextureRenderTarget(FName("SimGridRT"), BakedWaterSurfaceRT);
-	RiverSimSystem->SetVariableBool(FName("ReadCachedSim"), RenderState == EShallowWaterRenderState::BakedSim || RenderState == EShallowWaterRenderState::WaterComponentWithBakedSim);
+	RiverSimSystem->SetVariableBool(FName("ReadCachedSim"), RenderState == EShallowWaterRenderState::BakedSim || RenderState == EShallowWaterRenderState::WaterComponentWithBakedSim);	
 
 	RiverSimSystem->SetVariableFloat(FName("BottomContourCollisionDilation"), BottomContourCollisionDilation);
 	
@@ -510,6 +496,75 @@ void UShallowWaterRiverComponent::Rebuild()
 	bIsInitialized = true;
 }
 
+void UShallowWaterRiverComponent::AddActorsToRawArray(const TArray<TObjectPtr<AActor>> &ActorsArray, TArray<AActor*>& BottomContourActorsRawPtr)
+{
+	for (TObjectPtr<AActor> CurrActor : ActorsArray)
+	{
+		AActor* CurrActorRawPtr = CurrActor.Get();
+
+		// if we have a level instance, break it up and add each actor
+		if (ALevelInstance* LevelInstancePtr = Cast<ALevelInstance>(CurrActorRawPtr))
+		{
+			const ULevelInstanceSubsystem* LevelInstanceSubsystem = GetWorld()->GetSubsystem<ULevelInstanceSubsystem>();
+
+			LevelInstanceSubsystem->ForEachActorInLevelInstance(LevelInstancePtr, [&](AActor* SubActor)
+				{
+					BottomContourActorsRawPtr.Add(SubActor);
+					return true;
+				});
+		}
+		else
+		{
+			BottomContourActorsRawPtr.Add(CurrActorRawPtr);
+		}
+	}
+}
+
+void UShallowWaterRiverComponent::AddTaggedActorsToArray(TArray<FName> &TagsToUse, TArray<AActor*>& BottomContourActorsRawPtr)
+{
+	// if we have a tag set
+	// do an overlap test
+	//  filter by tag and add to the bottomcontour actors list
+	//  if a level instance is tagged, loop over the contained actors
+
+	if (!TagsToUse.IsEmpty())
+	{
+		FCollisionQueryParams Params(SCENE_QUERY_STAT(ShallowWaterRiverActorQuery), false);
+
+		TArray<FOverlapResult> Overlaps;
+		GetWorld()->OverlapMultiByChannel(Overlaps, SystemPos, FQuat::Identity, ECollisionChannel::ECC_WorldStatic,
+			FCollisionShape::MakeBox(0.5f * FVector(WorldGridSize.X, WorldGridSize.Y, 100000)), Params);
+
+		for (const FOverlapResult& OverlapResult : Overlaps)
+		{
+			if (UPrimitiveComponent* PrimitiveComponent = OverlapResult.GetComponent())
+			{
+				if (AActor* ComponentActor = PrimitiveComponent->GetOwner())
+				{
+					if (TagsToUse.ContainsByPredicate([&](const FName& Tag) { return Tag == NAME_None || ComponentActor->Tags.Contains(Tag); }))
+					{
+						// if we have a level instance, break it up and add each actor
+						if (ALevelInstance* LevelInstancePtr = Cast<ALevelInstance>(ComponentActor))
+						{
+							const ULevelInstanceSubsystem* LevelInstanceSubsystem = GetWorld()->GetSubsystem<ULevelInstanceSubsystem>();
+
+							LevelInstanceSubsystem->ForEachActorInLevelInstance(LevelInstancePtr, [&](AActor* SubActor)
+								{
+									BottomContourActorsRawPtr.Add(SubActor);
+									return true;
+								});
+						}
+						else
+						{
+							BottomContourActorsRawPtr.Add(ComponentActor);
+						}
+					}
+				}
+			}
+		}
+	}
+}
+
 void UShallowWaterRiverComponent::Bake()
 {
 	EObjectFlags TextureObjectFlags = EObjectFlags::RF_Public;
diff --git a/Engine/Plugins/Experimental/WaterAdvanced/Source/WaterAdvanced/Public/ShallowWaterRiverActor.h b/Engine/Plugins/Experimental/WaterAdvanced/Source/WaterAdvanced/Public/ShallowWaterRiverActor.h
--- a/Engine/Plugins/Experimental/WaterAdvanced/Source/WaterAdvanced/Public/ShallowWaterRiverActor.h
+++ b/Engine/Plugins/Experimental/WaterAdvanced/Source/WaterAdvanced/Public/ShallowWaterRiverActor.h
@@ -95,6 +95,12 @@ public:
 	UPROPERTY(EditAnywhere, Category = "Collisions")
 	float BottomContourCaptureOffset = 1000.f;
 
+	UPROPERTY(EditAnywhere, Category = "Collisions")
+	TArray<TObjectPtr<AActor>> DilatedBottomContourActors;
+
+	UPROPERTY(EditAnywhere, Category = "Collisions")
+	TArray<FName> DilatedBottomContourTags;		
+
 	UPROPERTY(EditAnywhere, Category = "Collisions")
 	float BottomContourCollisionDilation = 0.f;
 
@@ -110,6 +116,10 @@ public:
 #if WITH_EDITOR
 	void Rebuild();
 
+	void AddActorsToRawArray(const TArray<TObjectPtr<AActor>>& ActorsArray, TArray<AActor*>& BottomContourActorsRawPtr);
+
+	void AddTaggedActorsToArray(TArray<FName>& TagsToUse, TArray<AActor*>& BottomContourActorsRawPtr);
+
 	void Bake();
 
 	virtual void PostEditChangeProperty(FPropertyChangedEvent& PropertyChangedEvent) override;
