merge - multiple source/sinks for water bodies in shallow water river simulations

#virtualized

diff --git a/Engine/Plugins/Experimental/WaterAdvanced/Source/WaterAdvanced/Private/ShallowWaterRiverActor.cpp b/Engine/Plugins/Experimental/WaterAdvanced/Source/WaterAdvanced/Private/ShallowWaterRiverActor.cpp
--- a/Engine/Plugins/Experimental/WaterAdvanced/Source/WaterAdvanced/Private/ShallowWaterRiverActor.cpp
+++ b/Engine/Plugins/Experimental/WaterAdvanced/Source/WaterAdvanced/Private/ShallowWaterRiverActor.cpp
@@ -5,6 +5,7 @@
 #include "NiagaraComponent.h"
 #include "NiagaraSystem.h"
 #include "NiagaraFunctionLibrary.h"
+#include "NiagaraDataInterfaceArrayFunctionLibrary.h"
 
 #include "WaterBodyActor.h"
 #include "WaterSplineComponent.h"
@@ -193,95 +194,77 @@ void UShallowWaterRiverComponent::Rebuild()
 	AllWaterBodies.Empty();
 
 	// collect all the water bodies	
-	if (SourceRiverWaterBody != nullptr)
+	if (SourceRiverWaterBodies.Num() != 0)
 	{
-		AllWaterBodies.Add(SourceRiverWaterBody);
+		for (TObjectPtr<AWaterBody > CurrWaterBody : SourceRiverWaterBodies)
+		{
+			if (CurrWaterBody)
+			{
+				AllWaterBodies.Add(CurrWaterBody);
+			}
+			else 
+			{
+				UE_LOG(LogShallowWater, Warning, TEXT("UShallowWaterRiverComponent::Rebuild() - skipping null water body actor found"));
+				continue;
+			}
+		}
 	}
 	else	
 	{
-		UE_LOG(LogShallowWater, Warning, TEXT("UShallowWaterRiverComponent::Rebuild() - No source water body specified"));
+		UE_LOG(LogShallowWater, Warning, TEXT("UShallowWaterRiverComponent::Rebuild() - No source water bodies specified"));
 		return;
 	}
 	
-	if (SinkRiverWaterBody != nullptr)
-	{
-		AllWaterBodies.Add(SinkRiverWaterBody);
-	}	
-
-	for (TObjectPtr<AWaterBody > CurrWaterBody : AdditonalRiverWaterBodies)
-	{
-		AllWaterBodies.Add(CurrWaterBody);
-	}
-
 	if (AllWaterBodies.Num() == 0)
 	{
-		UE_LOG(LogShallowWater, Warning, TEXT("UShallowWaterRiverComponent::Rebuild() - No water bodies specified"));
+		UE_LOG(LogShallowWater, Warning, TEXT("UShallowWaterRiverComponent::Rebuild() - No valid source water bodies specified"));
 		return;
 	}
 
-	// accumulate bounding box for river water bodies
-	FBoxSphereBounds::Builder CombinedWorldBoundsBuilder;
-	for (TObjectPtr<AWaterBody > CurrWaterBody : AllWaterBodies)
+	bool HasValidSinks = false;
+	for (TObjectPtr<AWaterBody> CurrWaterBody : SinkRiverWaterBodies)
 	{
 		if (CurrWaterBody != nullptr)
 		{
-			TObjectPtr<UWaterBodyComponent> CurrWaterBodyComponent = CurrWaterBody->GetWaterBodyComponent();
-
-			if (CurrWaterBodyComponent != nullptr)
-			{
-				// accumulate bounds
-				FBoxSphereBounds WorldBounds;
-				CurrWaterBody->GetActorBounds(true, WorldBounds.Origin, WorldBounds.BoxExtent);				
-
-				CombinedWorldBoundsBuilder += WorldBounds;
-			}			
+			HasValidSinks = true;
+			AllWaterBodies.Add(CurrWaterBody);			
 		}
 		else
 		{
-			UE_LOG(LogShallowWater, Warning, TEXT("UShallowWaterRiverComponent::Rebuild() - skipping null water body actor found"));
+			UE_LOG(LogShallowWater, Warning, TEXT("UShallowWaterRiverComponent::Rebuild() - skipping null sink water body actor found"));
 			continue;
 		}
 	}
-	FBoxSphereBounds CombinedBounds(CombinedWorldBoundsBuilder);
 
-	if (CombinedBounds.BoxExtent.Length() < SMALL_NUMBER)
+	if (!HasValidSinks)
 	{
-		UE_LOG(LogShallowWater, Warning, TEXT("UShallowWaterRiverComponent::Rebuild() - river bodies have zero bounds"));
-		return;
+		UE_LOG(LogShallowWater, Warning, TEXT("UShallowWaterRiverComponent::Rebuild() - no valid sinks, using the first source as a sink"));
+		SinkRiverWaterBodies.Add(*AllWaterBodies.CreateConstIterator());
 	}
 
-	FVector SourcePos(0, 0, 0);
-	float SourceWidth = 1;
-	float SourceDepth = 1;
-	FVector SourceDir(1, 0, 0);	
-
-	// Get source
-	if (!QueryWaterAtSplinePoint(SourceRiverWaterBody, 0, SourcePos, SourceDir, SourceWidth, SourceDepth))
+	// accumulate bounding box for river water bodies
+	FBoxSphereBounds::Builder CombinedWorldBoundsBuilder;
+	for (TObjectPtr<AWaterBody > CurrWaterBody : AllWaterBodies)
 	{
-		UE_LOG(LogShallowWater, Warning, TEXT("UShallowWaterRiverComponent::Rebuild() - water source query failed"));
-		return;
-	}
-
-	FVector SinkPos(0, 0, 0);
-	float SinkWidth = 1;
-	float SinkDepth = 1;
-	FVector SinkDir(1, 0, 0);
+		TObjectPtr<UWaterBodyComponent> CurrWaterBodyComponent = CurrWaterBody->GetWaterBodyComponent();
 
-	// Get Sink		
-	TObjectPtr<AWaterBody> SinkToUse = SinkRiverWaterBody;
+		if (CurrWaterBodyComponent != nullptr)
+		{
+			// accumulate bounds
+			FBoxSphereBounds WorldBounds;
+			CurrWaterBody->GetActorBounds(true, WorldBounds.Origin, WorldBounds.BoxExtent);				
 
-	// if no sink specified, use source
-	if (SinkToUse == nullptr)
-	{
-		SinkToUse = SourceRiverWaterBody;
+			CombinedWorldBoundsBuilder += WorldBounds;
+		}
 	}
-	
-	if (!QueryWaterAtSplinePoint(SinkToUse, -1, SinkPos, SinkDir, SinkWidth, SinkDepth))
+	FBoxSphereBounds CombinedBounds(CombinedWorldBoundsBuilder);
+
+	if (CombinedBounds.BoxExtent.Length() < SMALL_NUMBER)
 	{
-		UE_LOG(LogShallowWater, Warning, TEXT("UShallowWaterRiverComponent::Rebuild() - water sink query failed"));
+		UE_LOG(LogShallowWater, Warning, TEXT("UShallowWaterRiverComponent::Rebuild() - river bodies have zero bounds"));
 		return;
 	}
-
+	
 	SystemPos = CombinedBounds.Origin - FVector(0, 0, CombinedBounds.BoxExtent.Z);
 	
 	RiverSimSystem = NewObject<UNiagaraComponent>(this, NAME_None, RF_Public);
@@ -311,7 +294,7 @@ void UShallowWaterRiverComponent::Rebuild()
 			UNiagaraFunctionLibrary::SetSceneCapture2DDataInterfaceManagedMode(RiverSimSystem, DIName,
 				ESceneCaptureSource::SCS_SceneDepth,
 				FIntPoint(ResolutionMaxAxis, ResolutionMaxAxis),
-				ETextureRenderTargetFormat::RTF_R16f,
+				ETextureRenderTargetFormat::RTF_R32f,
 				ECameraProjectionMode::Orthographic,
 				90.0f,
 				FMath::Max(WorldGridSize.X, WorldGridSize.Y),
@@ -350,7 +333,7 @@ void UShallowWaterRiverComponent::Rebuild()
 			UNiagaraFunctionLibrary::SetSceneCapture2DDataInterfaceManagedMode(RiverSimSystem, DilatedDIName,
 				ESceneCaptureSource::SCS_SceneDepth,
 				FIntPoint(ResolutionMaxAxis, ResolutionMaxAxis),
-				ETextureRenderTargetFormat::RTF_R16f,
+				ETextureRenderTargetFormat::RTF_R32f,
 				ECameraProjectionMode::Orthographic,
 				90.0f,
 				FMath::Max(WorldGridSize.X, WorldGridSize.Y),
@@ -450,15 +433,47 @@ void UShallowWaterRiverComponent::Rebuild()
 	float Overshoot = 1000.f;
 	float FinalSourceHeight = 2. * CombinedBounds.BoxExtent.Z + Overshoot;
 
-	RiverSimSystem->SetVariablePosition(FName("SourcePos"), SourcePos - FVector(0, 0, .5 * FinalSourceHeight) + FVector(SourceDir.X, SourceDir.Y, 0) * .5 * SourceSize);
-	RiverSimSystem->SetVariableVec3(FName("SourceSize"), FVector(SourceWidth, SourceSize, FinalSourceHeight));
-	RiverSimSystem->SetVariableFloat(FName("SourceAngle"), PI / 2.f + FMath::Acos(SourceDir.Dot(FVector(1,0,0))));
+	// Get sources
+	int32 i = 0;
+	for (TObjectPtr<AWaterBody> CurrWaterBody : SourceRiverWaterBodies)
+	{
+		FVector CurrSourcePos;
+		float CurrSourceWidth;
+		float CurrSourceDepth;
+		FVector CurrSourceDir;
+		if (!QueryWaterAtSplinePoint(CurrWaterBody, 0, CurrSourcePos, CurrSourceDir, CurrSourceWidth, CurrSourceDepth))
+		{
+			UE_LOG(LogShallowWater, Warning, TEXT("UShallowWaterRiverComponent::Rebuild() - water source query failed"));
+			continue;
+		}		
+		
+		FVector FullSourcePos = CurrSourcePos - FVector(0, 0, .5 * FinalSourceHeight) + FVector(CurrSourceDir.X, CurrSourceDir.Y, 0) * .5 * SourceSize;
+		UNiagaraDataInterfaceArrayFunctionLibrary::SetNiagaraArrayPositionValue(RiverSimSystem, "User.SourcePosArray", i, FullSourcePos, true);
+		UNiagaraDataInterfaceArrayFunctionLibrary::SetNiagaraArrayVectorValue(RiverSimSystem, "User.SourceSizeArray", i, FVector(CurrSourceWidth, SourceSize, FinalSourceHeight), true);
+		UNiagaraDataInterfaceArrayFunctionLibrary::SetNiagaraArrayFloatValue(RiverSimSystem, "User.SourceAngleArray", i, PI / 2.f + FMath::Acos(CurrSourceDir.Dot(FVector(1, 0, 0))), true);
+		i++;
+	}
+
+	// get sinks	
+	FVector SinkPos(0, 0, 0);
+	float SinkWidth = 1;
+	float SinkDepth = 1;
+	FVector SinkDir(1, 0, 0);
 	
-	// height of the sink box doesn't matter
-	float SinkBoxHeight = 100000;
-	RiverSimSystem->SetVariablePosition(FName("SinkPos"), SinkPos);
-	RiverSimSystem->SetVariableVec3(FName("SinkSize"), FVector(SinkWidth, SourceSize, SinkBoxHeight));
-	RiverSimSystem->SetVariableFloat(FName("SinkAngle"), PI / 2.f + FMath::Acos(SinkDir.Dot(FVector(1, 0, 0))));
+	for (TObjectPtr<AWaterBody> CurrWaterBody : SinkRiverWaterBodies)
+	{
+		if (!QueryWaterAtSplinePoint(CurrWaterBody, -1, SinkPos, SinkDir, SinkWidth, SinkDepth))
+		{
+			UE_LOG(LogShallowWater, Warning, TEXT("UShallowWaterRiverComponent::Rebuild() - water sink query failed"));
+			continue;
+		}
+
+		// height of the sink box doesn't matter
+		float SinkBoxHeight = 100000;	
+		UNiagaraDataInterfaceArrayFunctionLibrary::SetNiagaraArrayPositionValue(RiverSimSystem, "User.SinkPosArray", i, SinkPos, true);
+		UNiagaraDataInterfaceArrayFunctionLibrary::SetNiagaraArrayVectorValue(RiverSimSystem, "User.SinkSizeArray", i, FVector(SinkWidth, SourceSize, SinkBoxHeight), true);
+		UNiagaraDataInterfaceArrayFunctionLibrary::SetNiagaraArrayFloatValue(RiverSimSystem, "User.SinkAngleArray", i, PI / 2.f + FMath::Acos(SinkDir.Dot(FVector(1, 0, 0))), true);
+	}
 
 	RiverSimSystem->SetVariableFloat(FName("SimSpeed"), SimSpeed);
 	RiverSimSystem->SetVariableInt(FName("NumSteps"), NumSteps);	
@@ -477,7 +492,6 @@ void UShallowWaterRiverComponent::Rebuild()
 	{
 		RiverSimSystem->SetVariableTexture(FName("BakedSimTexture"), BakedWaterSurfaceTexture);
 	
-
 		FIntVector2 NumCells = FIntVector2(BakedWaterSurfaceTexture->GetSizeX(), BakedWaterSurfaceTexture->GetSizeY());
 
 		BakedWaterSurfaceNormalRT = NewObject<UTextureRenderTarget2D>(this);	
diff --git a/Engine/Plugins/Experimental/WaterAdvanced/Source/WaterAdvanced/Public/ShallowWaterRiverActor.h b/Engine/Plugins/Experimental/WaterAdvanced/Source/WaterAdvanced/Public/ShallowWaterRiverActor.h
--- a/Engine/Plugins/Experimental/WaterAdvanced/Source/WaterAdvanced/Public/ShallowWaterRiverActor.h
+++ b/Engine/Plugins/Experimental/WaterAdvanced/Source/WaterAdvanced/Public/ShallowWaterRiverActor.h
@@ -56,15 +56,12 @@ public:
 	UPROPERTY(EditAnywhere, Category = "Simulation", meta = (DisplayName = "Num Substeps"))
 	int NumSteps = 1;
 
-	UPROPERTY(EditAnywhere, Category = "Water", meta = (DisplayName = "Source River Water Body"))
-	TObjectPtr<AWaterBody> SourceRiverWaterBody;
+	UPROPERTY(EditAnywhere, Category = "Water", meta = (DisplayName = "Source River Water Bodies"))
+	TArray<TObjectPtr<AWaterBody>> SourceRiverWaterBodies;
 
-	UPROPERTY(EditAnywhere, Category = "Water", meta = (DisplayName = "Sink River Water Body"))
-	TObjectPtr<AWaterBody> SinkRiverWaterBody;
+	UPROPERTY(EditAnywhere, Category = "Water", meta = (DisplayName = "Sink River Water Bodies"))
+	TArray<TObjectPtr<AWaterBody>> SinkRiverWaterBodies;
 
-	UPROPERTY(EditAnywhere, Category = "Water", meta = (DisplayName = "Additional River Water Bodies"))
-	TArray<TObjectPtr<AWaterBody>> AdditonalRiverWaterBodies;
-	
 	UPROPERTY(EditAnywhere, Category = "Rendering", meta = (DisplayName = "Render State"))
 	TEnumAsByte<EShallowWaterRenderState> RenderState = EShallowWaterRenderState::WaterComponent;
 	
