refactor sw river material to be closer to the water body materials - make sure transitions work

#virtualized

diff --git a/Engine/Plugins/Experimental/WaterAdvanced/Source/WaterAdvanced/Private/ShallowWaterRiverActor.cpp b/Engine/Plugins/Experimental/WaterAdvanced/Source/WaterAdvanced/Private/ShallowWaterRiverActor.cpp
--- a/Engine/Plugins/Experimental/WaterAdvanced/Source/WaterAdvanced/Private/ShallowWaterRiverActor.cpp
+++ b/Engine/Plugins/Experimental/WaterAdvanced/Source/WaterAdvanced/Private/ShallowWaterRiverActor.cpp
@@ -8,6 +8,7 @@
 #include "NiagaraDataInterfaceArrayFunctionLibrary.h"
 
 #include "WaterBodyActor.h"
+#include "WaterBodyRiverComponent.h"
 #include "WaterSplineComponent.h"
 #include "WaterSubsystem.h"
 
@@ -830,9 +831,19 @@ void UShallowWaterRiverComponent::UpdateRenderState()
 {
 	if (BakedSimMaterial == nullptr)
 	{
-		BakedSimMaterial = LoadObject<UMaterialInstance>(nullptr, TEXT("/WaterAdvanced/Niagara/Materials/M_SW_River_Layered_Inst.M_SW_River_Layered_Inst"));
+		BakedSimMaterial = LoadObject<UMaterialInstance>(nullptr, TEXT("/WaterAdvanced/Niagara/Materials/SW_Water_Material_River.SW_Water_Material_River"));
 	}
 	
+	if (BakedSimRiverToLakeTransitionMaterial == nullptr)
+	{
+		BakedSimRiverToLakeTransitionMaterial = LoadObject<UMaterialInstance>(nullptr, TEXT("/WaterAdvanced/Niagara/Materials/SW_Water_Material_River_To_Lake_Transition.SW_Water_Material_River_To_Lake_Transition"));
+	}
+
+	if (BakedSimRiverToOceanTransitionMaterial == nullptr)
+	{
+		BakedSimRiverToOceanTransitionMaterial = LoadObject<UMaterialInstance>(nullptr, TEXT("/WaterAdvanced/Niagara/Materials/SW_Water_Material_River_To_Ocean_Transition.SW_Water_Material_River_To_Ocean_Transition"));
+	}
+
 	bool RenderWaterBody = RenderState == EShallowWaterRenderState::WaterComponent || RenderState == EShallowWaterRenderState::WaterComponentWithBakedSim;
 
 	if (RiverSimSystem != nullptr)
@@ -851,7 +862,7 @@ void UShallowWaterRiverComponent::UpdateRenderState()
 			continue;
 		}
 
-		TObjectPtr<UWaterBodyComponent> CurrWaterBodyComponent = CurrWaterBody->GetWaterBodyComponent();
+		TObjectPtr<UWaterBodyRiverComponent> CurrWaterBodyComponent = Cast<UWaterBodyRiverComponent>(CurrWaterBody->GetWaterBodyComponent());
 
 		if (CurrWaterBodyComponent != nullptr)
 		{
@@ -859,47 +870,36 @@ void UShallowWaterRiverComponent::UpdateRenderState()
 
 			if (RenderState == EShallowWaterRenderState::WaterComponentWithBakedSim)
 			{			
+				// override materials on water bodies				
+
 				if (OriginalWaterBodyMaterial == nullptr)
 				{
 					OriginalWaterBodyMaterial = CurrWaterBodyComponent->GetWaterMaterial();
 				}
 
-				CurrWaterBodyComponent->SetWaterMaterial(BakedSimMaterial);
-
-				UMaterialInstanceDynamic* WaterMID = CurrWaterBodyComponent->GetWaterMaterialInstance();
-				UMaterialInstanceDynamic* WaterInfoMID = CurrWaterBodyComponent->GetWaterInfoMaterialInstance();
-
-				// override materials on water bodies
-				
-				if (WaterMID)
+				if (OriginalWaterBodyRiverToLakeTransitionMaterial == nullptr)
 				{
-					WaterMID->SetTextureParameterValue("BakedWaterSimTex", BakedWaterSurfaceTexture);
-					WaterMID->SetVectorParameterValue("BakedWaterSimLocation", SystemPos);
-					WaterMID->SetVectorParameterValue("BakedWaterSimSize", FVector(WorldGridSize.X, WorldGridSize.Y, 1));
-
-					WaterMID->SetTextureParameterValue("BakedWaterSimNormalTex", BakedWaterSurfaceNormalRT);
-					WaterMID->SetTextureParameterValue("NormalDetailTex", NormalDetailRT);
-					WaterMID->SetTextureParameterValue("FoamTex", FoamRT);					
-				
-					float Dx;
-					if (WorldGridSize.X > WorldGridSize.Y)
-					{
-						Dx = WorldGridSize.X / ResolutionMaxAxis;
-					}
-					else					
-					{
-						Dx = WorldGridSize.Y / ResolutionMaxAxis;
-					}
-					WaterMID->SetScalarParameterValue("BakedWaterSimDx", Dx);
-
-					WaterMID->SetScalarParameterValue("UseBakedSimHack", RenderState == EShallowWaterRenderState::WaterComponentWithBakedSim ? 1 : 0);
+					OriginalWaterBodyRiverToLakeTransitionMaterial = CurrWaterBodyComponent->GetRiverToLakeTransitionMaterial();
 				}
-				else
+
+				if (OriginalWaterBodyRiverToOceanTransitionMaterial == nullptr)
 				{
-					UE_LOG(LogShallowWater, Warning, TEXT("UShallowWaterRiverComponent::UpdateRenderState() - Water Component MID is null"));
-					return;
+					OriginalWaterBodyRiverToOceanTransitionMaterial = CurrWaterBodyComponent->GetRiverToOceanTransitionMaterial();
 				}
 
+				CurrWaterBodyComponent->SetWaterMaterial(BakedSimMaterial);
+				UMaterialInstanceDynamic* WaterMID = CurrWaterBodyComponent->GetWaterMaterialInstance();			
+				SetWaterMIDParameters(WaterMID);
+
+				CurrWaterBodyComponent->SetLakeTransitionMaterial(BakedSimRiverToLakeTransitionMaterial);
+				UMaterialInstanceDynamic* WaterLakeTransitionMID = CurrWaterBodyComponent->GetRiverToLakeTransitionMaterialInstance();
+				SetWaterMIDParameters(WaterLakeTransitionMID);
+
+				CurrWaterBodyComponent->SetOceanTransitionMaterial(BakedSimRiverToOceanTransitionMaterial);
+				UMaterialInstanceDynamic* WaterOceanTransitionMID = CurrWaterBodyComponent->GetRiverToOceanTransitionMaterialInstance();
+				SetWaterMIDParameters(WaterOceanTransitionMID);
+												
+				UMaterialInstanceDynamic* WaterInfoMID = CurrWaterBodyComponent->GetWaterInfoMaterialInstance();
 				if (WaterInfoMID)
 				{
 					WaterInfoMID->SetTextureParameterValue("BakedWaterSimTex", BakedWaterSurfaceTexture);
@@ -919,6 +919,18 @@ void UShallowWaterRiverComponent::UpdateRenderState()
 					CurrWaterBodyComponent->SetWaterMaterial(OriginalWaterBodyMaterial.Get());
 					OriginalWaterBodyMaterial = nullptr;
 				}
+
+				if (OriginalWaterBodyRiverToLakeTransitionMaterial != nullptr)
+				{
+					CurrWaterBodyComponent->SetLakeTransitionMaterial(OriginalWaterBodyRiverToLakeTransitionMaterial.Get());
+					OriginalWaterBodyRiverToLakeTransitionMaterial = nullptr;
+				}
+
+				if (OriginalWaterBodyRiverToOceanTransitionMaterial != nullptr)
+				{
+					CurrWaterBodyComponent->SetOceanTransitionMaterial(OriginalWaterBodyRiverToOceanTransitionMaterial.Get());
+					OriginalWaterBodyRiverToOceanTransitionMaterial = nullptr;
+				}
 			}
 
 			CurrWaterBodyComponent->SetUseBakedSimulationForQueriesAndPhysics(
@@ -942,6 +954,38 @@ void UShallowWaterRiverComponent::UpdateRenderState()
 	}
 }
 
+void UShallowWaterRiverComponent::SetWaterMIDParameters(UMaterialInstanceDynamic* WaterMID)
+{
+	if (WaterMID)
+	{
+		WaterMID->SetTextureParameterValue("BakedWaterSimTex", BakedWaterSurfaceTexture);
+		WaterMID->SetVectorParameterValue("BakedWaterSimLocation", SystemPos);
+		WaterMID->SetVectorParameterValue("BakedWaterSimSize", FVector(WorldGridSize.X, WorldGridSize.Y, 1));
+
+		WaterMID->SetTextureParameterValue("BakedWaterSimNormalTex", BakedWaterSurfaceNormalRT);
+		WaterMID->SetTextureParameterValue("NormalDetailTex", NormalDetailRT);
+		WaterMID->SetTextureParameterValue("FoamTex", FoamRT);
+
+		float Dx;
+		if (WorldGridSize.X > WorldGridSize.Y)
+		{
+			Dx = WorldGridSize.X / ResolutionMaxAxis;
+		}
+		else
+		{
+			Dx = WorldGridSize.Y / ResolutionMaxAxis;
+		}
+		WaterMID->SetScalarParameterValue("BakedWaterSimDx", Dx);
+
+		WaterMID->SetScalarParameterValue("UseBakedSimHack", RenderState == EShallowWaterRenderState::WaterComponentWithBakedSim ? 1 : 0);
+	}
+	else
+	{
+		UE_LOG(LogShallowWater, Warning, TEXT("UShallowWaterRiverComponent::UpdateRenderState() - Water Component MID is null"));
+		return;
+	}
+}
+
 AShallowWaterRiver::AShallowWaterRiver(const FObjectInitializer& ObjectInitializer)
 	: Super(ObjectInitializer)
 {
diff --git a/Engine/Plugins/Experimental/WaterAdvanced/Source/WaterAdvanced/Public/ShallowWaterRiverActor.h b/Engine/Plugins/Experimental/WaterAdvanced/Source/WaterAdvanced/Public/ShallowWaterRiverActor.h
--- a/Engine/Plugins/Experimental/WaterAdvanced/Source/WaterAdvanced/Public/ShallowWaterRiverActor.h
+++ b/Engine/Plugins/Experimental/WaterAdvanced/Source/WaterAdvanced/Public/ShallowWaterRiverActor.h
@@ -68,6 +68,12 @@ public:
 	UPROPERTY(EditAnywhere, Category = "Rendering", meta = (DisplayName = "Baked Sim Material"))
 	TObjectPtr <class UMaterialInstance> BakedSimMaterial;
 
+	UPROPERTY(EditAnywhere, Category = "Rendering", meta = (DisplayName = "Baked Sim River To Lake Transition Material"))
+	TObjectPtr <class UMaterialInstance> BakedSimRiverToLakeTransitionMaterial;
+
+	UPROPERTY(EditAnywhere, Category = "Rendering", meta = (DisplayName = "Baked Sim River To Ocean Transition Material"))
+	TObjectPtr <class UMaterialInstance> BakedSimRiverToOceanTransitionMaterial;
+
 	UPROPERTY(EditAnywhere, Category = "Rendering", meta = (DisplayName = "Smoothing Width"))
 	int SmoothingWidth = 2;
 
@@ -110,6 +116,8 @@ public:
 
 	void UpdateRenderState();
 
+	void SetWaterMIDParameters(UMaterialInstanceDynamic* WaterMID);
+
 #if WITH_EDITOR
 	void Rebuild();
 
@@ -142,6 +150,12 @@ protected:
 	UPROPERTY(BlueprintReadOnly, VisibleDefaultsOnly, Category = "Shallow Water")
 	TObjectPtr <UMaterialInterface> OriginalWaterBodyMaterial;
 
+	UPROPERTY(BlueprintReadOnly, VisibleDefaultsOnly, Category = "Shallow Water")
+	TObjectPtr <UMaterialInterface> OriginalWaterBodyRiverToLakeTransitionMaterial;
+
+	UPROPERTY(BlueprintReadOnly, VisibleDefaultsOnly, Category = "Shallow Water")
+	TObjectPtr <UMaterialInterface> OriginalWaterBodyRiverToOceanTransitionMaterial;
+
 	bool QueryWaterAtSplinePoint(TObjectPtr<AWaterBody> WaterBody, int SplinePoint, FVector& OutPos, FVector& OutTangent, float& OutWidth, float& OutDepth);
 	
 private:
