Merge latest from external stream
Resave assets

#virtualized

diff --git a/Engine/Plugins/Experimental/WaterAdvanced/Source/WaterAdvanced/Private/ShallowWaterRiverActor.cpp b/Engine/Plugins/Experimental/WaterAdvanced/Source/WaterAdvanced/Private/ShallowWaterRiverActor.cpp
--- a/Engine/Plugins/Experimental/WaterAdvanced/Source/WaterAdvanced/Private/ShallowWaterRiverActor.cpp
+++ b/Engine/Plugins/Experimental/WaterAdvanced/Source/WaterAdvanced/Private/ShallowWaterRiverActor.cpp
@@ -155,6 +155,7 @@ void UShallowWaterRiverComponent::PostEditChangeProperty(FPropertyChangedEvent&
 void UShallowWaterRiverComponent::Rebuild()
 {	
 	bIsInitialized = false;
+	bTickInitialize = false;
 
 	if (NiagaraRiverSimulation == nullptr)
 	{
