fix various bugs with changing render state

#virtualized

diff --git a/Engine/Plugins/Experimental/WaterAdvanced/Source/WaterAdvanced/Private/ShallowWaterRiverActor.cpp b/Engine/Plugins/Experimental/WaterAdvanced/Source/WaterAdvanced/Private/ShallowWaterRiverActor.cpp
--- a/Engine/Plugins/Experimental/WaterAdvanced/Source/WaterAdvanced/Private/ShallowWaterRiverActor.cpp
+++ b/Engine/Plugins/Experimental/WaterAdvanced/Source/WaterAdvanced/Private/ShallowWaterRiverActor.cpp
@@ -143,13 +143,13 @@ void UShallowWaterRiverComponent::PostEditChangeProperty(FPropertyChangedEvent&
 	else
 	{
 		bIsInitialized = false;
-		bTickInitialize = false;
-
-		Rebuild();
+		bTickInitialize = false;		
 	}
 
-
+	Rebuild();
 	UpdateRenderState();
+	ReregisterComponent();
+
 }
 
 void UShallowWaterRiverComponent::Rebuild()
@@ -846,10 +846,14 @@ void UShallowWaterRiverComponent::UpdateRenderState()
 	}
 
 	bool RenderWaterBody = RenderState == EShallowWaterRenderState::WaterComponent || RenderState == EShallowWaterRenderState::WaterComponentWithBakedSim;
+	bool RenderSecondary = 
+		RenderState == EShallowWaterRenderState::WaterComponentWithBakedSim || 
+		RenderState == EShallowWaterRenderState::BakedSim || RenderState == EShallowWaterRenderState::LiveSim;
 
 	if (RiverSimSystem != nullptr)
 	{		
 		RiverSimSystem->SetVariableBool(FName("RenderWaterSurface"), !RenderWaterBody);
+		RiverSimSystem->SetVariableBool(FName("RenderSecondary"), RenderSecondary);
 		RiverSimSystem->SetVariableBool(FName("DebugRenderBottomContour"), RenderState == EShallowWaterRenderState::DebugRenderBottomContour);
 		RiverSimSystem->SetVariableBool(FName("DebugRenderFoam"), RenderState == EShallowWaterRenderState::DebugRenderFoam);
 		
