secondary system only computes emission points when in live sim mode
skip initializing scene captures in cache read mode

#virtualized

diff --git a/Engine/Plugins/Experimental/WaterAdvanced/Source/WaterAdvanced/Private/ShallowWaterRiverActor.cpp b/Engine/Plugins/Experimental/WaterAdvanced/Source/WaterAdvanced/Private/ShallowWaterRiverActor.cpp
--- a/Engine/Plugins/Experimental/WaterAdvanced/Source/WaterAdvanced/Private/ShallowWaterRiverActor.cpp
+++ b/Engine/Plugins/Experimental/WaterAdvanced/Source/WaterAdvanced/Private/ShallowWaterRiverActor.cpp
@@ -273,6 +273,8 @@ void UShallowWaterRiverComponent::Rebuild()
 	RiverSimSystem->bUseAttachParentBound = false;
 	RiverSimSystem->SetWorldLocation(SystemPos);
 
+	bool bReadBakedSim = RenderState == EShallowWaterRenderState::BakedSim || RenderState == EShallowWaterRenderState::WaterComponentWithBakedSim;
+
 	if (GetWorld() && GetWorld()->bIsWorldInitialized)
 	{
 		if (!RiverSimSystem->IsRegistered())
@@ -284,7 +286,7 @@ void UShallowWaterRiverComponent::Rebuild()
 		RiverSimSystem->SetAsset(NiagaraRiverSimulation);
 							
 		// convert to raw ptr array for function library
-		if (bUseCapture)
+		if (!bReadBakedSim && bUseCapture)
 		{
 			TArray<AActor*> BottomContourActorsRawPtr;
 			BottomContourActorsRawPtr.Add(nullptr);
@@ -483,7 +485,7 @@ void UShallowWaterRiverComponent::Rebuild()
 	BakedWaterSurfaceRT = NewObject<UTextureRenderTarget2D>(this, NAME_None, RF_Transient);
 	BakedWaterSurfaceRT->InitAutoFormat(1, 1);
 	RiverSimSystem->SetVariableTextureRenderTarget(FName("SimGridRT"), BakedWaterSurfaceRT);
-	RiverSimSystem->SetVariableBool(FName("ReadCachedSim"), RenderState == EShallowWaterRenderState::BakedSim || RenderState == EShallowWaterRenderState::WaterComponentWithBakedSim);	
+	RiverSimSystem->SetVariableBool(FName("ReadCachedSim"), bReadBakedSim);
 
 	RiverSimSystem->SetVariableFloat(FName("BottomContourCollisionDilation"), BottomContourCollisionDilation);
 	
