tweak some default values
- faster sim speed
- larger capture offset

diff --git a/Engine/Plugins/Experimental/WaterAdvanced/Source/WaterAdvanced/Public/ShallowWaterRiverActor.h b/Engine/Plugins/Experimental/WaterAdvanced/Source/WaterAdvanced/Public/ShallowWaterRiverActor.h
--- a/Engine/Plugins/Experimental/WaterAdvanced/Source/WaterAdvanced/Public/ShallowWaterRiverActor.h
+++ b/Engine/Plugins/Experimental/WaterAdvanced/Source/WaterAdvanced/Public/ShallowWaterRiverActor.h
@@ -51,10 +51,10 @@ public:
 	float SourceSize;
 
 	UPROPERTY(EditAnywhere, Category = "Simulation", meta = (DisplayName = "Speed"))
-	float SimSpeed = 1.f;
+	float SimSpeed = 10.f;
 
 	UPROPERTY(EditAnywhere, Category = "Simulation", meta = (DisplayName = "Num Substeps"))
-	int NumSteps = 1;
+	int NumSteps = 10;
 
 	UPROPERTY(EditAnywhere, Category = "Water", meta = (DisplayName = "Source River Water Bodies"))
 	TArray<TObjectPtr<AWaterBody>> SourceRiverWaterBodies;
@@ -96,7 +96,7 @@ public:
 	TArray<FName> BottomContourTags;
 
 	UPROPERTY(EditAnywhere, Category = "Collisions")
-	float BottomContourCaptureOffset = 1000.f;
+	float BottomContourCaptureOffset = 5000.f;
 
 	UPROPERTY(EditAnywhere, Category = "Collisions")
 	TArray<TObjectPtr<AActor>> DilatedBottomContourActors;
