change to soft pointers to address various validation errors with world partition

diff --git a/Engine/Plugins/Experimental/WaterAdvanced/Source/WaterAdvanced/Private/ShallowWaterRiverActor.cpp b/Engine/Plugins/Experimental/WaterAdvanced/Source/WaterAdvanced/Private/ShallowWaterRiverActor.cpp
--- a/Engine/Plugins/Experimental/WaterAdvanced/Source/WaterAdvanced/Private/ShallowWaterRiverActor.cpp
+++ b/Engine/Plugins/Experimental/WaterAdvanced/Source/WaterAdvanced/Private/ShallowWaterRiverActor.cpp
@@ -198,7 +198,7 @@ void UShallowWaterRiverComponent::Rebuild()
 	// collect all the water bodies	
 	if (SourceRiverWaterBodies.Num() != 0)
 	{
-		for (TObjectPtr<AWaterBody > CurrWaterBody : SourceRiverWaterBodies)
+		for (TSoftObjectPtr<AWaterBody > CurrWaterBody : SourceRiverWaterBodies)
 		{
 			if (CurrWaterBody)
 			{
@@ -224,7 +224,7 @@ void UShallowWaterRiverComponent::Rebuild()
 	}
 
 	bool HasValidSinks = false;
-	for (TObjectPtr<AWaterBody> CurrWaterBody : SinkRiverWaterBodies)
+	for (TSoftObjectPtr<AWaterBody> CurrWaterBody : SinkRiverWaterBodies)
 	{
 		if (CurrWaterBody != nullptr)
 		{
@@ -246,7 +246,7 @@ void UShallowWaterRiverComponent::Rebuild()
 
 	// accumulate bounding box for river water bodies
 	FBoxSphereBounds::Builder CombinedWorldBoundsBuilder;
-	for (TObjectPtr<AWaterBody > CurrWaterBody : AllWaterBodies)
+	for (TSoftObjectPtr<AWaterBody > CurrWaterBody : AllWaterBodies)
 	{
 		TObjectPtr<UWaterBodyComponent> CurrWaterBodyComponent = CurrWaterBody->GetWaterBodyComponent();
 
@@ -292,8 +292,7 @@ void UShallowWaterRiverComponent::Rebuild()
 			BottomContourActorsRawPtr.Add(nullptr);
 
 			AddActorsToRawArray(BottomContourActors, BottomContourActorsRawPtr);			
-			AddTaggedActorsToArray(BottomContourTags, BottomContourActorsRawPtr);			
-			
+			AddTaggedActorsToArray(BottomContourTags, BottomContourActorsRawPtr);
 			FName DIName = "User.BottomCapture";
 			UNiagaraFunctionLibrary::SetSceneCapture2DDataInterfaceManagedMode(RiverSimSystem, DIName,
 				ESceneCaptureSource::SCS_SceneDepth,
@@ -309,7 +308,7 @@ void UShallowWaterRiverComponent::Rebuild()
 	
 			// accumulate bounding box for river water bodies
 			FBoxSphereBounds::Builder BottomContourCombinedWorldBoundsBuilder;
-			for (TObjectPtr<AActor> BottomContourActor : BottomContourActors)
+			for (TSoftObjectPtr<AActor> BottomContourActor : BottomContourActors)
 			{
 				if (BottomContourActor != nullptr)
 				{
@@ -390,7 +389,7 @@ void UShallowWaterRiverComponent::Rebuild()
 	}
 
 	// look for the water info texture
-	for (TObjectPtr<AWaterBody > CurrWaterBody : AllWaterBodies)
+	for (TSoftObjectPtr<AWaterBody > CurrWaterBody : AllWaterBodies)
 	{
 		if (AWaterZone* WaterZone = CurrWaterBody->GetWaterBodyComponent()->GetWaterZone())
 		{			
@@ -439,7 +438,7 @@ void UShallowWaterRiverComponent::Rebuild()
 
 	// Get sources
 	int32 i = 0;
-	for (TObjectPtr<AWaterBody> CurrWaterBody : SourceRiverWaterBodies)
+	for (TSoftObjectPtr<AWaterBody> CurrWaterBody : SourceRiverWaterBodies)
 	{
 		FVector CurrSourcePos;
 		float CurrSourceWidth;
@@ -464,7 +463,7 @@ void UShallowWaterRiverComponent::Rebuild()
 	float SinkDepth = 1;
 	FVector SinkDir(1, 0, 0);
 	
-	for (TObjectPtr<AWaterBody> CurrWaterBody : SinkRiverWaterBodies)
+	for (TSoftObjectPtr<AWaterBody> CurrWaterBody : SinkRiverWaterBodies)
 	{
 		if (!QueryWaterAtSplinePoint(CurrWaterBody, -1, SinkPos, SinkDir, SinkWidth, SinkDepth))
 		{
@@ -514,9 +513,9 @@ void UShallowWaterRiverComponent::Rebuild()
 	bIsInitialized = true;
 }
 
-void UShallowWaterRiverComponent::AddActorsToRawArray(const TArray<TObjectPtr<AActor>> &ActorsArray, TArray<AActor*>& BottomContourActorsRawPtr)
+void UShallowWaterRiverComponent::AddActorsToRawArray(const TArray<TSoftObjectPtr<AActor>> &ActorsArray, TArray<AActor*>& BottomContourActorsRawPtr)
 {
-	for (TObjectPtr<AActor> CurrActor : ActorsArray)
+	for (TSoftObjectPtr<AActor> CurrActor : ActorsArray)
 	{
 		AActor* CurrActorRawPtr = CurrActor.Get();
 
@@ -639,7 +638,7 @@ void UShallowWaterRiverComponent::Bake()
 				
 		if (Depth > 1e-5)
 		{
-			for (TObjectPtr<AWaterBody > CurrWaterBody : AllWaterBodies)
+			for (TSoftObjectPtr<AWaterBody > CurrWaterBody : AllWaterBodies)
 			{						
 				TObjectPtr<UWaterBodyComponent> CurrWaterBodyComponent = CurrWaterBody->GetWaterBodyComponent();
 
@@ -678,7 +677,7 @@ void UShallowWaterRiverComponent::Bake()
 	}}
 
 	// set the sim texture on each water body that is in the simulated river.  
-	for (TObjectPtr<AWaterBody > CurrWaterBody : AllWaterBodies)
+	for (TSoftObjectPtr<AWaterBody > CurrWaterBody : AllWaterBodies)
 	{
 		TObjectPtr<UWaterBodyComponent> CurrWaterBodyComponent = CurrWaterBody->GetWaterBodyComponent();
 				
@@ -748,7 +747,7 @@ void UShallowWaterRiverComponent::Bake()
 	}
 }
 
-bool UShallowWaterRiverComponent::QueryWaterAtSplinePoint(TObjectPtr<AWaterBody> WaterBody, int SplinePoint, FVector& OutPos, FVector& OutTangent, float& OutWidth, float& OutDepth)
+bool UShallowWaterRiverComponent::QueryWaterAtSplinePoint(TSoftObjectPtr<AWaterBody> WaterBody, int SplinePoint, FVector& OutPos, FVector& OutTangent, float& OutWidth, float& OutDepth)
 {	
 	if (WaterBody != nullptr)
 	{
@@ -861,7 +860,7 @@ void UShallowWaterRiverComponent::UpdateRenderState()
 		
 	}
 
-	for (AWaterBody* CurrWaterBody : AllWaterBodies)
+	for (TSoftObjectPtr<AWaterBody > CurrWaterBody : AllWaterBodies)
 	{
 		if (!CurrWaterBody)
 		{
diff --git a/Engine/Plugins/Experimental/WaterAdvanced/Source/WaterAdvanced/Public/ShallowWaterRiverActor.h b/Engine/Plugins/Experimental/WaterAdvanced/Source/WaterAdvanced/Public/ShallowWaterRiverActor.h
--- a/Engine/Plugins/Experimental/WaterAdvanced/Source/WaterAdvanced/Public/ShallowWaterRiverActor.h
+++ b/Engine/Plugins/Experimental/WaterAdvanced/Source/WaterAdvanced/Public/ShallowWaterRiverActor.h
@@ -57,10 +57,10 @@ public:
 	int NumSteps = 10;
 
 	UPROPERTY(EditAnywhere, Category = "Water", meta = (DisplayName = "Source River Water Bodies"))
-	TArray<TObjectPtr<AWaterBody>> SourceRiverWaterBodies;
+	TArray<TSoftObjectPtr<AWaterBody>> SourceRiverWaterBodies;
 
 	UPROPERTY(EditAnywhere, Category = "Water", meta = (DisplayName = "Sink River Water Bodies"))
-	TArray<TObjectPtr<AWaterBody>> SinkRiverWaterBodies;
+	TArray<TSoftObjectPtr<AWaterBody>> SinkRiverWaterBodies;
 
 	UPROPERTY(EditAnywhere, Category = "Rendering", meta = (DisplayName = "Render State"))
 	TEnumAsByte<EShallowWaterRenderState> RenderState = EShallowWaterRenderState::WaterComponent;
@@ -90,7 +90,7 @@ public:
 	bool bUseCapture = true;
 
 	UPROPERTY(EditAnywhere, Category = "Collisions")
-	TArray<TObjectPtr<AActor>> BottomContourActors;
+	TArray<TSoftObjectPtr<AActor>> BottomContourActors;
 
 	UPROPERTY(EditAnywhere, Category = "Collisions")
 	TArray<FName> BottomContourTags;
@@ -99,7 +99,7 @@ public:
 	float BottomContourCaptureOffset = 5000.f;
 
 	UPROPERTY(EditAnywhere, Category = "Collisions")
-	TArray<TObjectPtr<AActor>> DilatedBottomContourActors;
+	TArray<TSoftObjectPtr<AActor>> DilatedBottomContourActors;
 
 	UPROPERTY(EditAnywhere, Category = "Collisions")
 	TArray<FName> DilatedBottomContourTags;		
@@ -121,7 +121,7 @@ public:
 #if WITH_EDITOR
 	void Rebuild();
 
-	void AddActorsToRawArray(const TArray<TObjectPtr<AActor>>& ActorsArray, TArray<AActor*>& BottomContourActorsRawPtr);
+	void AddActorsToRawArray(const TArray<TSoftObjectPtr<AActor>>& ActorsArray, TArray<AActor*>& BottomContourActorsRawPtr);
 
 	void AddTaggedActorsToArray(TArray<FName>& TagsToUse, TArray<AActor*>& BottomContourActorsRawPtr);
 
@@ -156,14 +156,14 @@ protected:
 	UPROPERTY(BlueprintReadOnly, VisibleDefaultsOnly, Category = "Shallow Water")
 	TObjectPtr <UMaterialInterface> OriginalWaterBodyRiverToOceanTransitionMaterial;
 
-	bool QueryWaterAtSplinePoint(TObjectPtr<AWaterBody> WaterBody, int SplinePoint, FVector& OutPos, FVector& OutTangent, float& OutWidth, float& OutDepth);
+	bool QueryWaterAtSplinePoint(TSoftObjectPtr<AWaterBody> WaterBody, int SplinePoint, FVector& OutPos, FVector& OutTangent, float& OutWidth, float& OutDepth);
 	
 private:
 	bool bIsInitialized;	
 	bool bTickInitialize;
-	
+
 	UPROPERTY()
-	TSet <TObjectPtr<AWaterBody>> AllWaterBodies;
+	TSet <TSoftObjectPtr<AWaterBody>> AllWaterBodies;
 
 	UPROPERTY()
 	FVector2D WorldGridSize;
