fix banding in texture advection due to incorrectly computed mip levels

diff --git a/Engine/Plugins/Experimental/WaterAdvanced/Shaders/WaterAdvanced.ush b/Engine/Plugins/Experimental/WaterAdvanced/Shaders/WaterAdvanced.ush
--- a/Engine/Plugins/Experimental/WaterAdvanced/Shaders/WaterAdvanced.ush
+++ b/Engine/Plugins/Experimental/WaterAdvanced/Shaders/WaterAdvanced.ush
@@ -111,16 +111,6 @@ float4 AdvectionInterpolation2(float2 UV, float2 Velocity, float t0, float t1, f
 	const float2 uv0 = AdvectedUV0 * AdvectedTextureScale;
 	const float2 uv1 = AdvectedUV1 * AdvectedTextureScale;
 
-	// compute mip levels
-	if (TexMipLevel.x < -1e-5 || TexMipLevel.y < 1e-5)
-	{
-		int2 TexDims;
-		int Levels;
-		Tex.GetDimensions(0, TexDims.x, TexDims.y, Levels);
-		TexMipLevel.x = ComputeMipLevel(uv0, TexDims);
-		TexMipLevel.y = ComputeMipLevel(uv1, TexDims);
-	}
-
 	// sample the texture for the two phases
 	const float4 TextureSample0 = Tex.SampleLevel(TexSampler, uv0, TexMipLevel.x);
 	const float4 TextureSample1 = Tex.SampleLevel(TexSampler, uv1, TexMipLevel.y);
@@ -130,7 +120,7 @@ float4 AdvectionInterpolation2(float2 UV, float2 Velocity, float t0, float t1, f
 }
 
 float4 AdvectionInterpolation(float2 UV, float2 Velocity, float t0, float t1, float t2, float lerp0, float lerp1, float lerp2,
-  float AdvectedTextureScale, float TexMipLevel, Texture2D<float4> Tex, sampler TexSampler)
+  float AdvectedTextureScale, float3 TexMipLevel, Texture2D<float4> Tex, sampler TexSampler)
 {
 	// advect uvs
 	const float2 AdvectedUV0 = UV - Velocity * t0;
@@ -141,20 +131,11 @@ float4 AdvectionInterpolation(float2 UV, float2 Velocity, float t0, float t1, fl
 	const float2 uv0 = AdvectedUV0 * AdvectedTextureScale;
 	const float2 uv1 = AdvectedUV1 * AdvectedTextureScale;
 	const float2 uv2 = AdvectedUV2 * AdvectedTextureScale;
-
-	// compute mips
-	if (TexMipLevel < -1e-5)
-	{	
-		int2 TexDims;
-		int Levels;
-		Tex.GetDimensions(0, TexDims.x, TexDims.y, Levels);
-		TexMipLevel = ComputeMipLevel(uv0, TexDims);
-	}
 		
 	// sample texture for each phase
-	const float4 TextureSample0 = Tex.SampleLevel(TexSampler, uv0, TexMipLevel);
-	const float4 TextureSample1 = Tex.SampleLevel(TexSampler, uv1, TexMipLevel);
-	const float4 TextureSample2 = Tex.SampleLevel(TexSampler, uv2, TexMipLevel);
+	const float4 TextureSample0 = Tex.SampleLevel(TexSampler, uv0, TexMipLevel.x);
+	const float4 TextureSample1 = Tex.SampleLevel(TexSampler, uv1, TexMipLevel.y);
+	const float4 TextureSample2 = Tex.SampleLevel(TexSampler, uv2, TexMipLevel.z);
 	
 	// blend according to phase
 	return TextureSample0 * lerp0 + TextureSample1 * lerp1 + TextureSample2 * lerp2;
