cache foam instead of computing in solver

#virtualized

diff --git a/Engine/Plugins/Experimental/WaterAdvanced/Source/WaterAdvanced/Private/ShallowWaterRiverActor.cpp b/Engine/Plugins/Experimental/WaterAdvanced/Source/WaterAdvanced/Private/ShallowWaterRiverActor.cpp
--- a/Engine/Plugins/Experimental/WaterAdvanced/Source/WaterAdvanced/Private/ShallowWaterRiverActor.cpp
+++ b/Engine/Plugins/Experimental/WaterAdvanced/Source/WaterAdvanced/Private/ShallowWaterRiverActor.cpp
@@ -484,6 +484,11 @@ void UShallowWaterRiverComponent::Rebuild()
 	BakedWaterSurfaceRT = NewObject<UTextureRenderTarget2D>(this, NAME_None, RF_Transient);
 	BakedWaterSurfaceRT->InitAutoFormat(1, 1);
 	RiverSimSystem->SetVariableTextureRenderTarget(FName("SimGridRT"), BakedWaterSurfaceRT);
+	
+	BakedFoamRT = NewObject<UTextureRenderTarget2D>(this, NAME_None, RF_Transient);
+	BakedFoamRT->InitAutoFormat(1, 1);
+	RiverSimSystem->SetVariableTextureRenderTarget(FName("FoamRT"), BakedFoamRT);
+	
 	RiverSimSystem->SetVariableBool(FName("ReadCachedSim"), bReadBakedSim);
 
 	RiverSimSystem->SetVariableFloat(FName("BottomContourCollisionDilation"), BottomContourCollisionDilation);
@@ -491,9 +496,10 @@ void UShallowWaterRiverComponent::Rebuild()
 	RiverSimSystem->SetVariableInt(FName("ExtrapolationHalfWidth"), SmoothingWidth);
 	RiverSimSystem->SetVariableFloat(FName("SmoothingHeightCutoff"), SmoothingCutoff);
 
-	if (BakedWaterSurfaceTexture != nullptr)
+	if (BakedWaterSurfaceTexture != nullptr && BakedFoamTexture != nullptr)
 	{
 		RiverSimSystem->SetVariableTexture(FName("BakedSimTexture"), BakedWaterSurfaceTexture);
+		RiverSimSystem->SetVariableTexture(FName("BakedFoamTexture"), BakedFoamTexture);
 	
 		FIntVector2 NumCells = FIntVector2(BakedWaterSurfaceTexture->GetSizeX(), BakedWaterSurfaceTexture->GetSizeY());
 
@@ -586,7 +592,7 @@ void UShallowWaterRiverComponent::Bake()
 {
 	EObjectFlags TextureObjectFlags = EObjectFlags::RF_Public;
 
-	if (!RiverSimSystem || !BakedWaterSurfaceRT)
+	if (!RiverSimSystem || !BakedWaterSurfaceRT || !BakedFoamRT)
 	{
 		UE_LOG(LogShallowWater, Warning, TEXT("UShallowWaterRiverComponent::Bake() - No simulation to bake"));
 		return;
@@ -621,6 +627,10 @@ void UShallowWaterRiverComponent::Bake()
 		ShallowWaterSimArrayValues[Index++] = FloatVal;
 	}
 
+	// bake foam and other data to texture
+	BakedFoamTexture = BakedFoamRT->ConstructTexture2D(this, "BakedFoamTexture", TextureObjectFlags);
+	RiverSimSystem->SetVariableTexture(FName("BakedFoamTexture"), BakedFoamTexture);
+	
 	BakedSim = NewObject<UBakedShallowWaterSimulationComponent>(this, NAME_None, RF_Public);
 	BakedSim->SimulationData = FShallowWaterSimulationGrid(ShallowWaterSimArrayValues, BakedWaterSurfaceTexture, FIntVector2(BakedWaterSurfaceRT->SizeX, BakedWaterSurfaceRT->SizeY), SystemPos, WorldGridSize);
 		
@@ -909,6 +919,7 @@ void UShallowWaterRiverComponent::UpdateRenderState()
 				if (WaterInfoMID)
 				{
 					WaterInfoMID->SetTextureParameterValue("BakedWaterSimTex", BakedWaterSurfaceTexture);
+					WaterInfoMID->SetTextureParameterValue("BakedFoamTex", BakedFoamTexture);
 					WaterInfoMID->SetVectorParameterValue("BakedWaterSimLocation", SystemPos);
 					WaterInfoMID->SetVectorParameterValue("BakedWaterSimSize", FVector(WorldGridSize.X, WorldGridSize.Y, 1));
 				}
@@ -965,6 +976,8 @@ void UShallowWaterRiverComponent::SetWaterMIDParameters(UMaterialInstanceDynamic
 	if (WaterMID)
 	{
 		WaterMID->SetTextureParameterValue("BakedWaterSimTex", BakedWaterSurfaceTexture);
+		WaterMID->SetTextureParameterValue("BakedFoamTex", BakedFoamTexture);
+
 		WaterMID->SetVectorParameterValue("BakedWaterSimLocation", SystemPos);
 		WaterMID->SetVectorParameterValue("BakedWaterSimSize", FVector(WorldGridSize.X, WorldGridSize.Y, 1));
 
diff --git a/Engine/Plugins/Experimental/WaterAdvanced/Source/WaterAdvanced/Public/ShallowWaterRiverActor.h b/Engine/Plugins/Experimental/WaterAdvanced/Source/WaterAdvanced/Public/ShallowWaterRiverActor.h
--- a/Engine/Plugins/Experimental/WaterAdvanced/Source/WaterAdvanced/Public/ShallowWaterRiverActor.h
+++ b/Engine/Plugins/Experimental/WaterAdvanced/Source/WaterAdvanced/Public/ShallowWaterRiverActor.h
@@ -83,6 +83,9 @@ public:
 	UPROPERTY(EditAnywhere, Category = "Shallow Water")
 	TObjectPtr<UTexture2D> BakedWaterSurfaceTexture;
 	
+	UPROPERTY(EditAnywhere, Category = "Shallow Water")
+	TObjectPtr<UTexture2D> BakedFoamTexture;
+
 	//UPROPERTY(EditAnywhere, Category = "Shallow Water")
 	//TObjectPtr<UTexture2D> SignedDistanceToSplineTexture;
 
@@ -144,6 +147,9 @@ protected:
 	UPROPERTY(BlueprintReadOnly, VisibleDefaultsOnly, Category = "Shallow Water")
 	TObjectPtr<UTextureRenderTarget2D> BakedWaterSurfaceRT;
 
+	UPROPERTY(BlueprintReadOnly, VisibleDefaultsOnly, Category = "Shallow Water")
+	TObjectPtr<UTextureRenderTarget2D> BakedFoamRT;
+
 	UPROPERTY(BlueprintReadOnly, VisibleDefaultsOnly, Category = "Shallow Water")
 	TObjectPtr<UBakedShallowWaterSimulationComponent> BakedSim;
 
