fixed bug with poorly initialized foam RTs

diff --git a/Engine/Plugins/Experimental/WaterAdvanced/Source/WaterAdvanced/Private/ShallowWaterRiverActor.cpp b/Engine/Plugins/Experimental/WaterAdvanced/Source/WaterAdvanced/Private/ShallowWaterRiverActor.cpp
--- a/Engine/Plugins/Experimental/WaterAdvanced/Source/WaterAdvanced/Private/ShallowWaterRiverActor.cpp
+++ b/Engine/Plugins/Experimental/WaterAdvanced/Source/WaterAdvanced/Private/ShallowWaterRiverActor.cpp
@@ -506,11 +506,7 @@ void UShallowWaterRiverComponent::Rebuild()
 		BakedWaterSurfaceNormalRT = NewObject<UTextureRenderTarget2D>(this);	
 		BakedWaterSurfaceNormalRT->InitAutoFormat(1, 1);
 		RiverSimSystem->SetVariableTextureRenderTarget(FName("NormalRT"), BakedWaterSurfaceNormalRT);
-	
-		FoamRT = NewObject<UTextureRenderTarget2D>(this);
-		FoamRT->InitAutoFormat(1, 1);
-		RiverSimSystem->SetVariableTextureRenderTarget(FName("FoamRT"), FoamRT);
-	
+
 		NormalDetailRT = NewObject<UTextureRenderTarget2D>(this);	
 		NormalDetailRT->InitAutoFormat(1, 1);
 		RiverSimSystem->SetVariableTextureRenderTarget(FName("NormalDetailRT"), NormalDetailRT);
@@ -919,7 +915,7 @@ void UShallowWaterRiverComponent::UpdateRenderState()
 				if (WaterInfoMID)
 				{
 					WaterInfoMID->SetTextureParameterValue("BakedWaterSimTex", BakedWaterSurfaceTexture);
-					WaterInfoMID->SetTextureParameterValue("BakedFoamTex", BakedFoamTexture);
+					WaterInfoMID->SetTextureParameterValue("FoamTex", BakedFoamTexture);
 					WaterInfoMID->SetVectorParameterValue("BakedWaterSimLocation", SystemPos);
 					WaterInfoMID->SetVectorParameterValue("BakedWaterSimSize", FVector(WorldGridSize.X, WorldGridSize.Y, 1));
 				}
@@ -976,7 +972,7 @@ void UShallowWaterRiverComponent::SetWaterMIDParameters(UMaterialInstanceDynamic
 	if (WaterMID)
 	{
 		WaterMID->SetTextureParameterValue("BakedWaterSimTex", BakedWaterSurfaceTexture);
-		WaterMID->SetTextureParameterValue("BakedFoamTex", BakedFoamTexture);
+		WaterMID->SetTextureParameterValue("FoamTex", BakedFoamTexture);
 
 		WaterMID->SetVectorParameterValue("BakedWaterSimLocation", SystemPos);
 		WaterMID->SetVectorParameterValue("BakedWaterSimSize", FVector(WorldGridSize.X, WorldGridSize.Y, 1));
