fix foam initialization

diff --git a/Engine/Plugins/Experimental/WaterAdvanced/Source/WaterAdvanced/Private/ShallowWaterRiverActor.cpp b/Engine/Plugins/Experimental/WaterAdvanced/Source/WaterAdvanced/Private/ShallowWaterRiverActor.cpp
--- a/Engine/Plugins/Experimental/WaterAdvanced/Source/WaterAdvanced/Private/ShallowWaterRiverActor.cpp
+++ b/Engine/Plugins/Experimental/WaterAdvanced/Source/WaterAdvanced/Private/ShallowWaterRiverActor.cpp
@@ -979,7 +979,6 @@ void UShallowWaterRiverComponent::SetWaterMIDParameters(UMaterialInstanceDynamic
 
 		WaterMID->SetTextureParameterValue("BakedWaterSimNormalTex", BakedWaterSurfaceNormalRT);
 		WaterMID->SetTextureParameterValue("NormalDetailTex", NormalDetailRT);
-		WaterMID->SetTextureParameterValue("FoamTex", FoamRT);
 
 		float Dx;
 		if (WorldGridSize.X > WorldGridSize.Y)
