fix compilation error
overlapping river water bodies (simulated + non simulated) render correctly

#virtualized

diff --git a/Engine/Plugins/Experimental/WaterAdvanced/Shaders/WaterAdvanced.ush b/Engine/Plugins/Experimental/WaterAdvanced/Shaders/WaterAdvanced.ush
--- a/Engine/Plugins/Experimental/WaterAdvanced/Shaders/WaterAdvanced.ush
+++ b/Engine/Plugins/Experimental/WaterAdvanced/Shaders/WaterAdvanced.ush
@@ -101,7 +101,7 @@ void Compute2Phases(const float PulseLength, const float Time,
 }
 
 float4 AdvectionInterpolation2(float2 UV, float2 Velocity, float t0, float t1, float lerpV,
-  float AdvectedTextureScale, float2 TexMipLevel, Texture2D<float4> Tex, sampler TexSampler, float2 CycleCount)
+  float AdvectedTextureScale, float2 TexMipLevel, Texture2D<float4> Tex, SamplerState TexSampler, float2 CycleCount)
 {
 	// advect uv, offset position somewhat randomly by the cycle count
 	const float2 AdvectedUV0 = UV - Velocity * t0 + CycleCount.x * 123;
@@ -120,7 +120,7 @@ float4 AdvectionInterpolation2(float2 UV, float2 Velocity, float t0, float t1, f
 }
 
 float4 AdvectionInterpolation(float2 UV, float2 Velocity, float t0, float t1, float t2, float lerp0, float lerp1, float lerp2,
-  float AdvectedTextureScale, float3 TexMipLevel, Texture2D<float4> Tex, sampler TexSampler)
+  float AdvectedTextureScale, float3 TexMipLevel, Texture2D<float4> Tex, SamplerState TexSampler)
 {
 	// advect uvs
 	const float2 AdvectedUV0 = UV - Velocity * t0;
