water body spline rendering mode works as expected

#virtualized

diff --git a/Engine/Plugins/Experimental/WaterAdvanced/Source/WaterAdvanced/Private/ShallowWaterRiverActor.cpp b/Engine/Plugins/Experimental/WaterAdvanced/Source/WaterAdvanced/Private/ShallowWaterRiverActor.cpp
--- a/Engine/Plugins/Experimental/WaterAdvanced/Source/WaterAdvanced/Private/ShallowWaterRiverActor.cpp
+++ b/Engine/Plugins/Experimental/WaterAdvanced/Source/WaterAdvanced/Private/ShallowWaterRiverActor.cpp
@@ -851,6 +851,11 @@ void UShallowWaterRiverComponent::UpdateRenderState()
 	{
 		BakedSimRiverToOceanTransitionMaterial = LoadObject<UMaterialInstance>(nullptr, TEXT("/WaterAdvanced/Niagara/Materials/SW_Water_Material_River_To_Ocean_Transition.SW_Water_Material_River_To_Ocean_Transition"));
 	}
+	
+	if (SplineRiverMaterial == nullptr)
+	{
+		SplineRiverMaterial = LoadObject<UMaterialInstance>(nullptr, TEXT("/WaterAdvanced/Niagara/Materials/SW_Water_Material_River_Spline.SW_Water_Material_River_Spline"));
+	}
 
 	bool RenderWaterBody = RenderState == EShallowWaterRenderState::WaterComponent || RenderState == EShallowWaterRenderState::WaterComponentWithBakedSim;
 	bool RenderSecondary = 
@@ -881,24 +886,7 @@ void UShallowWaterRiverComponent::UpdateRenderState()
 			CurrWaterBodyComponent->SetVisibility(RenderWaterBody);
 
 			if (RenderState == EShallowWaterRenderState::WaterComponentWithBakedSim)
-			{			
-				// override materials on water bodies				
-
-				if (OriginalWaterBodyMaterial == nullptr)
-				{
-					OriginalWaterBodyMaterial = CurrWaterBodyComponent->GetWaterMaterial();
-				}
-
-				if (OriginalWaterBodyRiverToLakeTransitionMaterial == nullptr)
-				{
-					OriginalWaterBodyRiverToLakeTransitionMaterial = CurrWaterBodyComponent->GetRiverToLakeTransitionMaterial();
-				}
-
-				if (OriginalWaterBodyRiverToOceanTransitionMaterial == nullptr)
-				{
-					OriginalWaterBodyRiverToOceanTransitionMaterial = CurrWaterBodyComponent->GetRiverToOceanTransitionMaterial();
-				}
-
+			{				
 				CurrWaterBodyComponent->SetWaterMaterial(BakedSimMaterial);
 				UMaterialInstanceDynamic* WaterMID = CurrWaterBodyComponent->GetWaterMaterialInstance();			
 				SetWaterMIDParameters(WaterMID);
@@ -927,22 +915,9 @@ void UShallowWaterRiverComponent::UpdateRenderState()
 			}
 			else if (RenderState == EShallowWaterRenderState::WaterComponent)
 			{
-				if (OriginalWaterBodyMaterial != nullptr)
-				{
-					CurrWaterBodyComponent->SetWaterMaterial(OriginalWaterBodyMaterial.Get());
-					OriginalWaterBodyMaterial = nullptr;
-				}
-
-				if (OriginalWaterBodyRiverToLakeTransitionMaterial != nullptr)
-				{
-					CurrWaterBodyComponent->SetLakeTransitionMaterial(OriginalWaterBodyRiverToLakeTransitionMaterial.Get());
-					OriginalWaterBodyRiverToLakeTransitionMaterial = nullptr;
-				}
-
-				if (OriginalWaterBodyRiverToOceanTransitionMaterial != nullptr)
+				if (SplineRiverMaterial != nullptr)
 				{
-					CurrWaterBodyComponent->SetOceanTransitionMaterial(OriginalWaterBodyRiverToOceanTransitionMaterial.Get());
-					OriginalWaterBodyRiverToOceanTransitionMaterial = nullptr;
+					CurrWaterBodyComponent->SetWaterMaterial(SplineRiverMaterial);					
 				}
 			}
 
diff --git a/Engine/Plugins/Experimental/WaterAdvanced/Source/WaterAdvanced/Public/ShallowWaterRiverActor.h b/Engine/Plugins/Experimental/WaterAdvanced/Source/WaterAdvanced/Public/ShallowWaterRiverActor.h
--- a/Engine/Plugins/Experimental/WaterAdvanced/Source/WaterAdvanced/Public/ShallowWaterRiverActor.h
+++ b/Engine/Plugins/Experimental/WaterAdvanced/Source/WaterAdvanced/Public/ShallowWaterRiverActor.h
@@ -74,6 +74,9 @@ public:
 	UPROPERTY(EditAnywhere, Category = "Rendering", meta = (DisplayName = "Baked Sim River To Ocean Transition Material"))
 	TObjectPtr <class UMaterialInstance> BakedSimRiverToOceanTransitionMaterial;
 
+	UPROPERTY(EditAnywhere, Category = "Rendering", meta = (DisplayName = "Spline River Material"))
+	TObjectPtr <class UMaterialInstance> SplineRiverMaterial;
+
 	UPROPERTY(EditAnywhere, Category = "Rendering", meta = (DisplayName = "Smoothing Width"))
 	int SmoothingWidth = 2;
 
@@ -153,15 +156,6 @@ protected:
 	UPROPERTY(BlueprintReadOnly, VisibleDefaultsOnly, Category = "Shallow Water")
 	TObjectPtr<UBakedShallowWaterSimulationComponent> BakedSim;
 
-	UPROPERTY(BlueprintReadOnly, VisibleDefaultsOnly, Category = "Shallow Water")
-	TObjectPtr <UMaterialInterface> OriginalWaterBodyMaterial;
-
-	UPROPERTY(BlueprintReadOnly, VisibleDefaultsOnly, Category = "Shallow Water")
-	TObjectPtr <UMaterialInterface> OriginalWaterBodyRiverToLakeTransitionMaterial;
-
-	UPROPERTY(BlueprintReadOnly, VisibleDefaultsOnly, Category = "Shallow Water")
-	TObjectPtr <UMaterialInterface> OriginalWaterBodyRiverToOceanTransitionMaterial;
-
 	bool QueryWaterAtSplinePoint(TSoftObjectPtr<AWaterBody> WaterBody, int SplinePoint, FVector& OutPos, FVector& OutTangent, float& OutWidth, float& OutDepth);
 	
 private:
