expose spline matching parameters on the shallow water actor

#virtualized

diff --git a/Engine/Plugins/Experimental/WaterAdvanced/Source/WaterAdvanced/Private/ShallowWaterRiverActor.cpp b/Engine/Plugins/Experimental/WaterAdvanced/Source/WaterAdvanced/Private/ShallowWaterRiverActor.cpp
--- a/Engine/Plugins/Experimental/WaterAdvanced/Source/WaterAdvanced/Private/ShallowWaterRiverActor.cpp
+++ b/Engine/Plugins/Experimental/WaterAdvanced/Source/WaterAdvanced/Private/ShallowWaterRiverActor.cpp
@@ -479,7 +479,11 @@ void UShallowWaterRiverComponent::Rebuild()
 	}
 
 	RiverSimSystem->SetVariableFloat(FName("SimSpeed"), SimSpeed);
-	RiverSimSystem->SetVariableInt(FName("NumSteps"), NumSteps);	
+	RiverSimSystem->SetVariableInt(FName("NumSteps"), NumSteps);
+
+	RiverSimSystem->SetVariableBool(FName("MatchSpline"), bMatchSpline);
+	RiverSimSystem->SetVariableFloat(FName("RemoveOutsideSplineAmount"), RemoveOutsideSplineAmount);
+	RiverSimSystem->SetVariableFloat(FName("SplineHeightMatchingAmount"), MatchSplineHeightAmount);
 	
 	BakedWaterSurfaceRT = NewObject<UTextureRenderTarget2D>(this, NAME_None, RF_Transient);
 	BakedWaterSurfaceRT->InitAutoFormat(1, 1);
diff --git a/Engine/Plugins/Experimental/WaterAdvanced/Source/WaterAdvanced/Public/ShallowWaterRiverActor.h b/Engine/Plugins/Experimental/WaterAdvanced/Source/WaterAdvanced/Public/ShallowWaterRiverActor.h
--- a/Engine/Plugins/Experimental/WaterAdvanced/Source/WaterAdvanced/Public/ShallowWaterRiverActor.h
+++ b/Engine/Plugins/Experimental/WaterAdvanced/Source/WaterAdvanced/Public/ShallowWaterRiverActor.h
@@ -56,6 +56,15 @@ public:
 	UPROPERTY(EditAnywhere, Category = "Simulation", meta = (DisplayName = "Num Substeps"))
 	int NumSteps = 10;
 
+	UPROPERTY(EditAnywhere, Category = "Simulation", meta = (DisplayName = "Match Spline"))
+	bool bMatchSpline = true;
+
+	UPROPERTY(EditAnywhere, Category = "Simulation", meta = (DisplayName = "Match Spline Height Amount", EditCondition = "bMatchSpline"))
+	float MatchSplineHeightAmount = 5.;
+
+	UPROPERTY(EditAnywhere, Category = "Simulation", meta = (DisplayName = "Remove Outside Spline Amount", EditCondition = "bMatchSpline"))
+	float RemoveOutsideSplineAmount = 1000.;
+
 	UPROPERTY(EditAnywhere, Category = "Water", meta = (DisplayName = "Source River Water Bodies"))
 	TArray<TSoftObjectPtr<AWaterBody>> SourceRiverWaterBodies;
 
