fix initialization for case where water sim doesn't render properly after opening level

diff --git a/Engine/Plugins/Experimental/WaterAdvanced/Source/WaterAdvanced/Private/ShallowWaterRiverActor.cpp b/Engine/Plugins/Experimental/WaterAdvanced/Source/WaterAdvanced/Private/ShallowWaterRiverActor.cpp
--- a/Engine/Plugins/Experimental/WaterAdvanced/Source/WaterAdvanced/Private/ShallowWaterRiverActor.cpp
+++ b/Engine/Plugins/Experimental/WaterAdvanced/Source/WaterAdvanced/Private/ShallowWaterRiverActor.cpp
@@ -42,6 +42,7 @@ UShallowWaterRiverComponent::UShallowWaterRiverComponent(const FObjectInitialize
 
 	bIsInitialized = false;
 	bTickInitialize = false;
+	bRenderStateTickInitialize = false;
 
 	ResolutionMaxAxis = 512;
 	SourceSize = 1000;
@@ -77,10 +78,11 @@ void UShallowWaterRiverComponent::PostLoad()
 	}
 	else
 	{
+		RiverSimSystem->ReinitializeSystem();
 		RiverSimSystem->Activate();
 	}
 
-	UpdateRenderState();
+	bRenderStateTickInitialize = false;		
 }
 
 void UShallowWaterRiverComponent::TickComponent(float DeltaTime, ELevelTick TickType, FActorComponentTickFunction* ThisTickFunction)
@@ -109,12 +111,19 @@ void UShallowWaterRiverComponent::TickComponent(float DeltaTime, ELevelTick Tick
 		// UE_LOG(LogShallowWater, Warning, TEXT("UShallowWaterRiverComponent::TickComponent() - null Niagara sim when trying to activate. Please reset."));
 	}
 #endif
+
+	if (!bRenderStateTickInitialize)
+	{
+		UpdateRenderState();
+	}
 }
 
 void UShallowWaterRiverComponent::BeginPlay()
 {
 	Super::BeginPlay();
 
+	bRenderStateTickInitialize = false;
+
 	UpdateRenderState();
 }
 
@@ -146,6 +155,8 @@ void UShallowWaterRiverComponent::PostEditChangeProperty(FPropertyChangedEvent&
 		bTickInitialize = false;		
 	}
 
+	bRenderStateTickInitialize = false;
+
 	Rebuild();
 	UpdateRenderState();
 	ReregisterComponent();
@@ -880,7 +891,7 @@ void UShallowWaterRiverComponent::UpdateRenderState()
 		if (!CurrWaterBody)
 		{
 			UE_LOG(LogShallowWater, Warning, TEXT("UShallowWaterRiverComponent::UpdateRenderState() - Water Body Actor is null"));
-			continue;
+			return;
 		}
 
 		TObjectPtr<UWaterBodyRiverComponent> CurrWaterBodyComponent = Cast<UWaterBodyRiverComponent>(CurrWaterBody->GetWaterBodyComponent());
@@ -944,6 +955,10 @@ void UShallowWaterRiverComponent::UpdateRenderState()
 			*/			
 		}
 	}
+
+	bRenderStateTickInitialize = true;
+
+	RiverSimSystem->ResetSystem();
 }
 
 void UShallowWaterRiverComponent::SetWaterMIDParameters(UMaterialInstanceDynamic* WaterMID)
diff --git a/Engine/Plugins/Experimental/WaterAdvanced/Source/WaterAdvanced/Public/ShallowWaterRiverActor.h b/Engine/Plugins/Experimental/WaterAdvanced/Source/WaterAdvanced/Public/ShallowWaterRiverActor.h
--- a/Engine/Plugins/Experimental/WaterAdvanced/Source/WaterAdvanced/Public/ShallowWaterRiverActor.h
+++ b/Engine/Plugins/Experimental/WaterAdvanced/Source/WaterAdvanced/Public/ShallowWaterRiverActor.h
@@ -170,6 +170,7 @@ protected:
 private:
 	bool bIsInitialized;	
 	bool bTickInitialize;
+	bool bRenderStateTickInitialize;
 
 	UPROPERTY()
 	TSet <TSoftObjectPtr<AWaterBody>> AllWaterBodies;
