fixed nullptr

diff --git a/Engine/Plugins/Experimental/WaterAdvanced/Source/WaterAdvanced/Private/ShallowWaterRiverActor.cpp b/Engine/Plugins/Experimental/WaterAdvanced/Source/WaterAdvanced/Private/ShallowWaterRiverActor.cpp
--- a/Engine/Plugins/Experimental/WaterAdvanced/Source/WaterAdvanced/Private/ShallowWaterRiverActor.cpp
+++ b/Engine/Plugins/Experimental/WaterAdvanced/Source/WaterAdvanced/Private/ShallowWaterRiverActor.cpp
@@ -958,7 +958,10 @@ void UShallowWaterRiverComponent::UpdateRenderState()
 
 	bRenderStateTickInitialize = true;
 
-	RiverSimSystem->ResetSystem();
+	if (RiverSimSystem != nullptr)
+	{	
+		RiverSimSystem->ResetSystem();
+	}
 }
 
 void UShallowWaterRiverComponent::SetWaterMIDParameters(UMaterialInstanceDynamic* WaterMID)
