new water system
fix source/sink angle

diff --git a/Engine/Plugins/Experimental/WaterAdvanced/Source/WaterAdvanced/Private/ShallowWaterRiverActor.cpp b/Engine/Plugins/Experimental/WaterAdvanced/Source/WaterAdvanced/Private/ShallowWaterRiverActor.cpp
--- a/Engine/Plugins/Experimental/WaterAdvanced/Source/WaterAdvanced/Private/ShallowWaterRiverActor.cpp
+++ b/Engine/Plugins/Experimental/WaterAdvanced/Source/WaterAdvanced/Private/ShallowWaterRiverActor.cpp
@@ -31,6 +31,9 @@
 #include "LevelInstance/LevelInstanceActor.h"
 #include "LevelInstance/LevelInstanceSubsystem.h"
 
+bool bShallowWaterRiverDebugVisualize = false;
+FAutoConsoleVariableRef CVarShallowWaterRiverDebugVisualize(TEXT("r.ShallowWater.RiverDebugVisualize"), bShallowWaterRiverDebugVisualize, TEXT(""));
+
 UShallowWaterRiverComponent::UShallowWaterRiverComponent(const FObjectInitializer& Initializer)
 	: Super(Initializer)
 {
@@ -249,6 +252,11 @@ void UShallowWaterRiverComponent::Rebuild()
 		}
 	}
 
+	// flush all debug draw lines
+#if ENABLE_DRAW_DEBUG
+	FlushPersistentDebugLines(GetWorld());
+#endif
+
 	if (!HasValidSinks)
 	{
 		UE_LOG(LogShallowWater, Warning, TEXT("UShallowWaterRiverComponent::Rebuild() - no valid sinks, using the first source as a sink"));
@@ -462,9 +470,34 @@ void UShallowWaterRiverComponent::Rebuild()
 		}		
 		
 		FVector FullSourcePos = CurrSourcePos - FVector(0, 0, .5 * FinalSourceHeight) + FVector(CurrSourceDir.X, CurrSourceDir.Y, 0) * .5 * SourceSize;
+		FVector FullSourceSize = FVector(CurrSourceWidth, SourceSize, FinalSourceHeight); 
+		
+		CurrSourceDir = FVector(CurrSourceDir.X, CurrSourceDir.Y, 0);
+		CurrSourceDir.Normalize();
+		
+		FVector BaseVector = {0,1,0};	
+		double FullSourceAngle =  FMath::Acos(FVector::DotProduct(BaseVector, CurrSourceDir));
+		
+		FVector AxisToUse = FVector::CrossProduct(BaseVector, CurrSourceDir);
+		AxisToUse.Normalize();
+
+#if ENABLE_DRAW_DEBUG
+		if (bShallowWaterRiverDebugVisualize)
+		{
+			FQuat TmpQ = FQuat::MakeFromRotationVector(AxisToUse * FullSourceAngle);		
+			DrawDebugBox(GetWorld() , (FVector) FullSourcePos, .5 * FullSourceSize, TmpQ, FColor::Green, true);		
+		}
+#endif
+
+		// flip axis so we don't need to store the vector itself
+		if (AxisToUse.Z < 0)
+		{
+			FullSourceAngle *= -1;
+		}
+
 		UNiagaraDataInterfaceArrayFunctionLibrary::SetNiagaraArrayPositionValue(RiverSimSystem, "User.SourcePosArray", i, FullSourcePos, true);
-		UNiagaraDataInterfaceArrayFunctionLibrary::SetNiagaraArrayVectorValue(RiverSimSystem, "User.SourceSizeArray", i, FVector(CurrSourceWidth, SourceSize, FinalSourceHeight), true);
-		UNiagaraDataInterfaceArrayFunctionLibrary::SetNiagaraArrayFloatValue(RiverSimSystem, "User.SourceAngleArray", i, PI / 2.f + FMath::Acos(CurrSourceDir.Dot(FVector(1, 0, 0))), true);
+		UNiagaraDataInterfaceArrayFunctionLibrary::SetNiagaraArrayVectorValue(RiverSimSystem, "User.SourceSizeArray", i, FullSourceSize, true);
+		UNiagaraDataInterfaceArrayFunctionLibrary::SetNiagaraArrayFloatValue(RiverSimSystem, "User.SourceAngleArray", i, FullSourceAngle, true);
 		i++;
 	}
 
@@ -483,10 +516,39 @@ void UShallowWaterRiverComponent::Rebuild()
 		}
 
 		// height of the sink box doesn't matter
-		float SinkBoxHeight = 100000;	
+		float SinkBoxHeight = 100000;
+		FVector FullSinkSize = FVector(SinkWidth, SourceSize, SinkBoxHeight);
+		
+		SinkDir = FVector(SinkDir.X, SinkDir.Y, 0);
+		SinkDir.Normalize();
+		
+		FVector BaseVector = {0,1,0};	
+		double FullSinkAngle =  FMath::Acos(FVector::DotProduct(BaseVector, SinkDir));
+		
+
+		FVector AxisToUse = FVector::CrossProduct(BaseVector, SinkDir);
+		AxisToUse.Normalize();
+
+#if ENABLE_DRAW_DEBUG
+		if (bShallowWaterRiverDebugVisualize)
+		{
+			FQuat TmpQ = FQuat::MakeFromRotationVector(AxisToUse * FullSinkAngle);		
+			DrawDebugBox(GetWorld() , (FVector) SinkPos, .5 * FullSinkSize, TmpQ, FColor::Red, true);		
+		}
+#endif
+
+		// flip axis so we don't need to store the vector itself
+		if (AxisToUse.Z < 0)
+		{
+			FullSinkAngle *= -1;
+		}
+
 		UNiagaraDataInterfaceArrayFunctionLibrary::SetNiagaraArrayPositionValue(RiverSimSystem, "User.SinkPosArray", i, SinkPos, true);
-		UNiagaraDataInterfaceArrayFunctionLibrary::SetNiagaraArrayVectorValue(RiverSimSystem, "User.SinkSizeArray", i, FVector(SinkWidth, SourceSize, SinkBoxHeight), true);
-		UNiagaraDataInterfaceArrayFunctionLibrary::SetNiagaraArrayFloatValue(RiverSimSystem, "User.SinkAngleArray", i, PI / 2.f + FMath::Acos(SinkDir.Dot(FVector(1, 0, 0))), true);
+
+		UNiagaraDataInterfaceArrayFunctionLibrary::SetNiagaraArrayVectorValue(RiverSimSystem, "User.SinkSizeArray", i, FullSinkSize, true);
+		UNiagaraDataInterfaceArrayFunctionLibrary::SetNiagaraArrayFloatValue(RiverSimSystem, "User.SinkAngleArray", i, FullSinkAngle, true);		
+
+		i++;
 	}
 
 	RiverSimSystem->SetVariableFloat(FName("SimSpeed"), SimSpeed);
@@ -789,15 +851,14 @@ bool UShallowWaterRiverComponent::QueryWaterAtSplinePoint(TSoftObjectPtr<AWaterB
 			if (Metadata != nullptr)
 			{
 				OutPos = CurrSpline->SplineCurves.Position.Points[SplinePoint].OutVal;
-				OutPos = WaterBody->GetActorTransform().TransformPosition(OutPos);
+				OutPos = CurrSpline->GetComponentTransform().TransformPosition(OutPos);
 
 				OutWidth = Metadata->RiverWidth.Points[SplinePoint].OutVal;
 				OutDepth = Metadata->Depth.Points[SplinePoint].OutVal;
 
 				OutTangent = CurrSpline->SplineCurves.Position.Points[SplinePoint].LeaveTangent;
 
-				OutTangent = WaterBody->GetActorTransform().TransformVector(OutTangent);
-
+				OutTangent = CurrSpline->GetComponentTransform().TransformVector(OutTangent);
 				OutTangent.Normalize();
 			}
 			else
