water component rendering mode supports transition materials

#virtualized

diff --git a/Engine/Plugins/Experimental/WaterAdvanced/Source/WaterAdvanced/Private/ShallowWaterRiverActor.cpp b/Engine/Plugins/Experimental/WaterAdvanced/Source/WaterAdvanced/Private/ShallowWaterRiverActor.cpp
--- a/Engine/Plugins/Experimental/WaterAdvanced/Source/WaterAdvanced/Private/ShallowWaterRiverActor.cpp
+++ b/Engine/Plugins/Experimental/WaterAdvanced/Source/WaterAdvanced/Private/ShallowWaterRiverActor.cpp
@@ -933,6 +933,16 @@ void UShallowWaterRiverComponent::UpdateRenderState()
 		SplineRiverMaterial = LoadObject<UMaterialInstance>(nullptr, TEXT("/WaterAdvanced/Niagara/Materials/SW_Water_Material_River_Spline.SW_Water_Material_River_Spline"));
 	}
 
+	if (SplineRiverToLakeTransitionMaterial == nullptr)
+	{
+		SplineRiverToLakeTransitionMaterial = LoadObject<UMaterialInstance>(nullptr, TEXT("/WaterAdvanced/Niagara/Materials/SW_Water_Material_River_To_Lake_Transition_Spline.SW_Water_Material_River_To_Lake_Transition_Spline"));
+	}
+
+	if (SplineRiverToOceanTransitionMaterial == nullptr)
+	{
+		SplineRiverToOceanTransitionMaterial = LoadObject<UMaterialInstance>(nullptr, TEXT("/WaterAdvanced/Niagara/Materials/SW_Water_Material_River_To_Ocean_Transition_Spline.SW_Water_Material_River_To_Ocean_Transition_Spline"));
+	}
+
 	bool RenderWaterBody = RenderState == EShallowWaterRenderState::WaterComponent || RenderState == EShallowWaterRenderState::WaterComponentWithBakedSim;
 	bool RenderSecondary = 
 		RenderState == EShallowWaterRenderState::WaterComponentWithBakedSim || 
@@ -991,10 +1001,9 @@ void UShallowWaterRiverComponent::UpdateRenderState()
 			}
 			else if (RenderState == EShallowWaterRenderState::WaterComponent)
 			{
-				if (SplineRiverMaterial != nullptr)
-				{
-					CurrWaterBodyComponent->SetWaterMaterial(SplineRiverMaterial);					
-				}
+				CurrWaterBodyComponent->SetWaterMaterial(SplineRiverMaterial);
+				CurrWaterBodyComponent->SetLakeTransitionMaterial(SplineRiverToLakeTransitionMaterial);
+				CurrWaterBodyComponent->SetOceanTransitionMaterial(SplineRiverToOceanTransitionMaterial);
 			}
 
 			CurrWaterBodyComponent->SetUseBakedSimulationForQueriesAndPhysics(
diff --git a/Engine/Plugins/Experimental/WaterAdvanced/Source/WaterAdvanced/Public/ShallowWaterRiverActor.h b/Engine/Plugins/Experimental/WaterAdvanced/Source/WaterAdvanced/Public/ShallowWaterRiverActor.h
--- a/Engine/Plugins/Experimental/WaterAdvanced/Source/WaterAdvanced/Public/ShallowWaterRiverActor.h
+++ b/Engine/Plugins/Experimental/WaterAdvanced/Source/WaterAdvanced/Public/ShallowWaterRiverActor.h
@@ -86,6 +86,12 @@ public:
 	UPROPERTY(EditAnywhere, Category = "Rendering", meta = (DisplayName = "Spline River Material"))
 	TObjectPtr <class UMaterialInstance> SplineRiverMaterial;
 
+	UPROPERTY(EditAnywhere, Category = "Rendering", meta = (DisplayName = "Spline River To Lake Transition Material"))
+	TObjectPtr <class UMaterialInstance> SplineRiverToLakeTransitionMaterial;
+
+	UPROPERTY(EditAnywhere, Category = "Rendering", meta = (DisplayName = "Spline River To Ocean Transition Material"))
+	TObjectPtr <class UMaterialInstance> SplineRiverToOceanTransitionMaterial;
+
 	UPROPERTY(EditAnywhere, Category = "Rendering", meta = (DisplayName = "Smoothing Width"))
 	int SmoothingWidth = 2;
 
