- 3 phase option for bracketing
- remove option for bracketing+magnitude (just make magnitude delta 0)
- fix magnitude delta == 0
- fix angle delta == 0
- reparameterize angle delta to match magnitude delta
- fixed angle mapping

#virtualized

diff --git a/Engine/Plugins/Experimental/WaterAdvanced/Shaders/WaterAdvanced.ush b/Engine/Plugins/Experimental/WaterAdvanced/Shaders/WaterAdvanced.ush
--- a/Engine/Plugins/Experimental/WaterAdvanced/Shaders/WaterAdvanced.ush
+++ b/Engine/Plugins/Experimental/WaterAdvanced/Shaders/WaterAdvanced.ush
@@ -10,33 +10,57 @@ static const float TWOPI = 2. * PI;
 
 void BracketAngle(float2 Velocity, float BracketAngleDelta, out float2 VelocityDir0, out float2 VelocityDir1, out float BlendAlpha)
 {
-	const float AngleDelta = PI / BracketAngleDelta;
-	const float2x2 RotateByAngleDelta = float2x2(cos(AngleDelta), sin(AngleDelta), -sin(AngleDelta), cos(AngleDelta));
+	if (BracketAngleDelta > 1e-5)
+	{
+		const float AngleDelta = BracketAngleDelta * PI / 180.0;
+		const float2x2 RotateByAngleDelta = float2x2(cos(AngleDelta), sin(AngleDelta), -sin(AngleDelta), cos(AngleDelta));
 
-    // direction of input velocity
-	const float Angle = atan2(Velocity.y, Velocity.x) + 2.0 * PI;
+		// direction of input velocity
+		//const float Angle = atan2(Velocity.y, Velocity.x) + 2.0 * PI;
 
-    // Snap to min bracketed angle
-	const float Fractional = fmod(Angle, AngleDelta);
-	const float Angle0 = Angle - Fractional;
-    
-    // Compute velocity direction snapped to the bracketed angle
-	VelocityDir0 = float2(cos(Angle0), sin(Angle0));
+		float Angle = atan2(Velocity.y, Velocity.x);
 
-    // second bracketed direction is just rotated by the angle delta
-	VelocityDir1 = mul(VelocityDir0, RotateByAngleDelta);
+Angle = Angle < 0 ? Angle + 2.0 * PI : Angle;
 
-    // blend is the distance the velocity heading is to the first bracketed angle
-	BlendAlpha = Fractional / AngleDelta;
+		// Snap to min bracketed angle
+		const float Fractional = fmod(Angle, AngleDelta);
+		const float Angle0 = Angle - Fractional;
+    
+		// Compute velocity direction snapped to the bracketed angle
+		VelocityDir0 = float2(cos(Angle0), sin(Angle0));
+
+		// second bracketed direction is just rotated by the angle delta
+		VelocityDir1 = mul(VelocityDir0, RotateByAngleDelta);
+
+		// blend is the distance the velocity heading is to the first bracketed angle
+		BlendAlpha = Fractional / AngleDelta;
+	}
+	else
+	{
+		// edge case where angle delta is 0, this turns off bracketing
+		VelocityDir0 = normalize(Velocity);
+		VelocityDir1 = VelocityDir0;		
+		BlendAlpha = 0.0;
+	}
 }
 
   // More general form of bracketing - takes just a float and a bracket size.
 void BracketValue(float Value, float BracketMagnitudeDelta, out float Value0, out float Value1, out float BlendAlpha)
 {	
-	const float Fractional = fmod(Value, BracketMagnitudeDelta);
-	Value0 = Value - Fractional;
-	Value1 = Value0 + BracketMagnitudeDelta;
-	BlendAlpha = Fractional / BracketMagnitudeDelta;
+	if (BracketMagnitudeDelta > 1e-5)
+	{
+		const float Fractional = fmod(Value, BracketMagnitudeDelta);
+		Value0 = Value - Fractional;
+		Value1 = Value0 + BracketMagnitudeDelta;
+		BlendAlpha = Fractional / BracketMagnitudeDelta;
+	}
+	else
+	{
+		// edge case where magnitude delta is 0, this turns off bracketing
+		Value0 = Value;
+		Value1 = Value;
+		BlendAlpha = 0.;
+	}
 }
 
 float2 CalculateUV(float2 vAxis, float2 position, float Time, float Speed)
@@ -67,9 +91,8 @@ float2 ComputeMipLevel(float2 uv, int2 sz)
 
 void Compute3Phases(const float PulseLength, const float Time, 
 	out float lerp0, out float lerp1, out float lerp2,
-	out float t0, out float t1, out float t2)
-{
-	const float HalfPulseLength = .5 * PulseLength;
+	out float t0, out float t1, out float t2, out float3 CycleCount)
+{	
 	const float TimeScale = TWOPI / PulseLength;
 
 	const int NumPhases = 3;
@@ -85,6 +108,11 @@ void Compute3Phases(const float PulseLength, const float Time,
 	t0 = (Time + PulseLength * 0. / NumPhases) % PulseLength;
 	t1 = (Time + PulseLength * 1. / NumPhases) % PulseLength;
 	t2 = (Time + PulseLength * 2. / NumPhases) % PulseLength;
+
+	CycleCount = float3(
+		floor((Time + PulseLength * 0. / NumPhases) / PulseLength), 
+		floor((Time + PulseLength * 1. / NumPhases) / PulseLength),
+		floor((Time + PulseLength * 2. / NumPhases) / PulseLength));
 }
 
 void Compute2Phases(const float PulseLength, const float Time,
@@ -119,13 +147,13 @@ float4 AdvectionInterpolation2(float2 UV, float2 Velocity, float t0, float t1, f
 	return lerp(TextureSample0, TextureSample1, lerpV);	
 }
 
-float4 AdvectionInterpolation(float2 UV, float2 Velocity, float t0, float t1, float t2, float lerp0, float lerp1, float lerp2,
-  float AdvectedTextureScale, float3 TexMipLevel, Texture2D<float4> Tex, SamplerState TexSampler)
+float4 AdvectionInterpolation3(float2 UV, float2 Velocity, float t0, float t1, float t2, float lerp0, float lerp1, float lerp2,
+  float AdvectedTextureScale, float3 TexMipLevel, Texture2D<float4> Tex, SamplerState TexSampler, float3 CycleCount)
 {
 	// advect uvs
-	const float2 AdvectedUV0 = UV - Velocity * t0;
-	const float2 AdvectedUV1 = UV - Velocity * t1;
-	const float2 AdvectedUV2 = UV - Velocity * t2;
+	const float2 AdvectedUV0 = UV - Velocity * t0 + CycleCount.x * 123;
+	const float2 AdvectedUV1 = UV - Velocity * t1 + CycleCount.y * 123;
+	const float2 AdvectedUV2 = UV - Velocity * t2 + CycleCount.z * 123;
 
 	// scale uvs
 	const float2 uv0 = AdvectedUV0 * AdvectedTextureScale;
