offset by cycle count via random numbers rather than constant offset to reduce repetition and flashing

diff --git a/Engine/Plugins/Experimental/WaterAdvanced/Shaders/WaterAdvanced.ush b/Engine/Plugins/Experimental/WaterAdvanced/Shaders/WaterAdvanced.ush
--- a/Engine/Plugins/Experimental/WaterAdvanced/Shaders/WaterAdvanced.ush
+++ b/Engine/Plugins/Experimental/WaterAdvanced/Shaders/WaterAdvanced.ush
@@ -15,12 +15,9 @@ void BracketAngle(float2 Velocity, float BracketAngleDelta, out float2 VelocityD
 		const float AngleDelta = BracketAngleDelta * PI / 180.0;
 		const float2x2 RotateByAngleDelta = float2x2(cos(AngleDelta), sin(AngleDelta), -sin(AngleDelta), cos(AngleDelta));
 
-		// direction of input velocity
-		//const float Angle = atan2(Velocity.y, Velocity.x) + 2.0 * PI;
-
+		// direction of input velocity		
 		float Angle = atan2(Velocity.y, Velocity.x);
-
-Angle = Angle < 0 ? Angle + 2.0 * PI : Angle;
+		Angle = Angle < 0.0 ? Angle + 2.0 * PI : Angle;
 
 		// Snap to min bracketed angle
 		const float Fractional = fmod(Angle, AngleDelta);
@@ -132,8 +129,11 @@ float4 AdvectionInterpolation2(float2 UV, float2 Velocity, float t0, float t1, f
   float AdvectedTextureScale, float2 TexMipLevel, Texture2D<float4> Tex, SamplerState TexSampler, float2 CycleCount)
 {
 	// advect uv, offset position somewhat randomly by the cycle count
-	const float2 AdvectedUV0 = UV - Velocity * t0 + CycleCount.x * 123;
-	const float2 AdvectedUV1 = UV - Velocity * t1 + CycleCount.y * 123;
+	float OffsetRandX = PseudoRandom(float2(CycleCount.x, 0))*1024.;
+	float OffsetRandY = PseudoRandom(float2(CycleCount.y, 0))*1024.;
+
+	const float2 AdvectedUV0 = UV - Velocity * t0 + OffsetRandX;
+	const float2 AdvectedUV1 = UV - Velocity * t1 + OffsetRandY;
 
 	// scale uv
 	const float2 uv0 = AdvectedUV0 * AdvectedTextureScale;
@@ -151,9 +151,13 @@ float4 AdvectionInterpolation3(float2 UV, float2 Velocity, float t0, float t1, f
   float AdvectedTextureScale, float3 TexMipLevel, Texture2D<float4> Tex, SamplerState TexSampler, float3 CycleCount)
 {
 	// advect uvs
-	const float2 AdvectedUV0 = UV - Velocity * t0 + CycleCount.x * 123;
-	const float2 AdvectedUV1 = UV - Velocity * t1 + CycleCount.y * 123;
-	const float2 AdvectedUV2 = UV - Velocity * t2 + CycleCount.z * 123;
+	float OffsetRandX = PseudoRandom(float2(CycleCount.x, 0))*1024.;
+	float OffsetRandY = PseudoRandom(float2(CycleCount.y, 0))*1024.;
+	float OffsetRandZ = PseudoRandom(float2(CycleCount.z, 0))*1024.;
+
+	const float2 AdvectedUV0 = UV - Velocity * t0 + OffsetRandX;
+	const float2 AdvectedUV1 = UV - Velocity * t1 + OffsetRandY;
+	const float2 AdvectedUV2 = UV - Velocity * t2 + OffsetRandZ;
 
 	// scale uvs
 	const float2 uv0 = AdvectedUV0 * AdvectedTextureScale;
