Fix initialization for ShallowWaterRiverActor when set in Water Component mode

diff --git a/Engine/Plugins/Experimental/WaterAdvanced/Source/WaterAdvanced/Private/ShallowWaterRiverActor.cpp b/Engine/Plugins/Experimental/WaterAdvanced/Source/WaterAdvanced/Private/ShallowWaterRiverActor.cpp
--- a/Engine/Plugins/Experimental/WaterAdvanced/Source/WaterAdvanced/Private/ShallowWaterRiverActor.cpp
+++ b/Engine/Plugins/Experimental/WaterAdvanced/Source/WaterAdvanced/Private/ShallowWaterRiverActor.cpp
@@ -128,6 +128,15 @@ void UShallowWaterRiverComponent::BeginPlay()
 	bRenderStateTickInitialize = false;
 
 	UpdateRenderState();
+
+	// make sure the simulation is not going to be run in case of various initialization edge cases
+	bool bReadBakedSim = RenderState == EShallowWaterRenderState::BakedSim || RenderState == EShallowWaterRenderState::WaterComponentWithBakedSim || RenderState == EShallowWaterRenderState::WaterComponent;	
+	if (RiverSimSystem != nullptr && bReadBakedSim)
+	{
+		RiverSimSystem->SetVariableBool(FName("ReadCachedSim"), bReadBakedSim);
+		RiverSimSystem->ReinitializeSystem();
+		RiverSimSystem->Activate();
+	}
 }
 
 void UShallowWaterRiverComponent::OnUnregister()
@@ -150,7 +159,8 @@ void UShallowWaterRiverComponent::PostEditChangeProperty(FPropertyChangedEvent&
 	// this should go before rebuild not after...something is wrong
 	if (PropertyName == GET_MEMBER_NAME_CHECKED(UShallowWaterRiverComponent, RenderState) && RiverSimSystem != nullptr && RiverSimSystem->IsActive())
 	{
-		RiverSimSystem->SetVariableBool(FName("ReadCachedSim"), RenderState == EShallowWaterRenderState::BakedSim || RenderState == EShallowWaterRenderState::WaterComponentWithBakedSim);
+		bool bReadBakedSim = RenderState == EShallowWaterRenderState::BakedSim || RenderState == EShallowWaterRenderState::WaterComponentWithBakedSim || RenderState == EShallowWaterRenderState::WaterComponent;	
+		RiverSimSystem->SetVariableBool(FName("ReadCachedSim"), bReadBakedSim);
 	}
 	else
 	{
@@ -292,7 +302,7 @@ void UShallowWaterRiverComponent::Rebuild()
 	RiverSimSystem->bUseAttachParentBound = false;
 	RiverSimSystem->SetWorldLocation(SystemPos);
 
-	bool bReadBakedSim = RenderState == EShallowWaterRenderState::BakedSim || RenderState == EShallowWaterRenderState::WaterComponentWithBakedSim;
+	bool bReadBakedSim = RenderState == EShallowWaterRenderState::BakedSim || RenderState == EShallowWaterRenderState::WaterComponentWithBakedSim || RenderState == EShallowWaterRenderState::WaterComponent;
 
 	if (GetWorld() && GetWorld()->bIsWorldInitialized)
 	{
@@ -943,6 +953,7 @@ void UShallowWaterRiverComponent::UpdateRenderState()
 		SplineRiverToOceanTransitionMaterial = LoadObject<UMaterialInstance>(nullptr, TEXT("/WaterAdvanced/Niagara/Materials/SW_Water_Material_River_To_Ocean_Transition_Spline.SW_Water_Material_River_To_Ocean_Transition_Spline"));
 	}
 
+	bool bReadBakedSim = RenderState == EShallowWaterRenderState::BakedSim || RenderState == EShallowWaterRenderState::WaterComponentWithBakedSim || RenderState == EShallowWaterRenderState::WaterComponent;
 	bool RenderWaterBody = RenderState == EShallowWaterRenderState::WaterComponent || RenderState == EShallowWaterRenderState::WaterComponentWithBakedSim;
 	bool RenderSecondary = 
 		RenderState == EShallowWaterRenderState::WaterComponentWithBakedSim || 
@@ -954,7 +965,7 @@ void UShallowWaterRiverComponent::UpdateRenderState()
 		RiverSimSystem->SetVariableBool(FName("RenderSecondary"), RenderSecondary);
 		RiverSimSystem->SetVariableBool(FName("DebugRenderBottomContour"), RenderState == EShallowWaterRenderState::DebugRenderBottomContour);
 		RiverSimSystem->SetVariableBool(FName("DebugRenderFoam"), RenderState == EShallowWaterRenderState::DebugRenderFoam);
-		
+		RiverSimSystem->SetVariableBool(FName("ReadCachedSim"), bReadBakedSim);
 	}
 
 	for (TSoftObjectPtr<AWaterBody > CurrWaterBody : AllWaterBodies)
