shared normal map w/ ocean emitter singleton actor

#virtualized

diff --git a/Engine/Plugins/Experimental/WaterAdvanced/Source/WaterAdvanced/Private/FFTOceanPatchActor.cpp b/Engine/Plugins/Experimental/WaterAdvanced/Source/WaterAdvanced/Private/FFTOceanPatchActor.cpp
new file mode 100644
--- /dev/null
+++ b/Engine/Plugins/Experimental/WaterAdvanced/Source/WaterAdvanced/Private/FFTOceanPatchActor.cpp
@@ -0,0 +1,67 @@
+// Copyright Epic Games, Inc. All Rights Reserved.
+
+#include "FFTOceanPatchActor.h"
+
+#include "NiagaraComponent.h"
+#include "NiagaraSystem.h"
+#include "NiagaraFunctionLibrary.h"
+#include "NiagaraDataInterfaceArrayFunctionLibrary.h"
+
+#include "WaterBodyActor.h"
+#include "WaterBodyRiverComponent.h"
+#include "WaterSplineComponent.h"
+#include "WaterSubsystem.h"
+
+#include "BakedShallowWaterSimulationComponent.h"
+#include "Engine/TextureRenderTarget2D.h"
+#include "Engine/TextureRenderTarget2DArray.h"
+#include "Materials/MaterialInstanceDynamic.h"
+#include "Components/SplineMeshComponent.h"
+#include "PhysicsEngine/ConvexElem.h"
+#include "PhysicsEngine/BodySetup.h"
+#include "Engine/OverlapResult.h"
+
+#include "TextureResource.h"
+#include "ShallowWaterCommon.h"
+#include "Engine/World.h"
+#include "Engine/Texture2D.h"
+#include "Math/Float16Color.h"
+#include "Kismet/GameplayStatics.h"
+#include "Landscape.h"
+#include "LevelInstance/LevelInstanceActor.h"
+#include "LevelInstance/LevelInstanceSubsystem.h"
+
+AFFTOceanPatchActor::AFFTOceanPatchActor(const FObjectInitializer& ObjectInitializer)
+	: Super(ObjectInitializer)
+{	
+	SetHidden(false);
+}
+
+void AFFTOceanPatchActor::PostLoad()
+{
+
+	if (NiagaraOceanSimulation == nullptr)
+	{
+		NiagaraOceanSimulation = LoadObject<UNiagaraSystem>(nullptr, TEXT("/WaterAdvanced/Niagara/Systems/Grid2D_OceanPatch.Grid2D_OceanPatch"));
+	}
+
+	FFTOceanSystem = NewObject<UNiagaraComponent>(this, NAME_None, RF_Public);	
+	
+	if (GetWorld() && GetWorld()->bIsWorldInitialized)
+	{
+		if (!FFTOceanSystem->IsRegistered())
+		{
+			FFTOceanSystem->RegisterComponentWithWorld(GetWorld());
+		}
+
+		FFTOceanSystem->SetVisibleFlag(true);
+		FFTOceanSystem->SetAsset(NiagaraOceanSimulation);
+
+		OceanNormalRT = NewObject<UTextureRenderTarget2D>(this, NAME_None, RF_Transient);
+		OceanNormalRT->InitAutoFormat(1, 1);
+		FFTOceanSystem->SetVariableTextureRenderTarget(FName("OceanNormalRT"), OceanNormalRT);
+		FFTOceanSystem->ReinitializeSystem();
+		FFTOceanSystem->ResetSystem();
+		FFTOceanSystem->Activate();
+	}
+}
diff --git a/Engine/Plugins/Experimental/WaterAdvanced/Source/WaterAdvanced/Private/ShallowWaterRiverActor.cpp b/Engine/Plugins/Experimental/WaterAdvanced/Source/WaterAdvanced/Private/ShallowWaterRiverActor.cpp
--- a/Engine/Plugins/Experimental/WaterAdvanced/Source/WaterAdvanced/Private/ShallowWaterRiverActor.cpp
+++ b/Engine/Plugins/Experimental/WaterAdvanced/Source/WaterAdvanced/Private/ShallowWaterRiverActor.cpp
@@ -31,6 +31,8 @@
 #include "LevelInstance/LevelInstanceActor.h"
 #include "LevelInstance/LevelInstanceSubsystem.h"
 
+#include "EngineUtils.h"
+
 bool bShallowWaterRiverDebugVisualize = false;
 FAutoConsoleVariableRef CVarShallowWaterRiverDebugVisualize(TEXT("r.ShallowWater.RiverDebugVisualize"), bShallowWaterRiverDebugVisualize, TEXT(""));
 
@@ -66,6 +68,37 @@ UShallowWaterRiverComponent::UShallowWaterRiverComponent(const FObjectInitialize
 
 }
 
+TObjectPtr<AFFTOceanPatchActor> UShallowWaterRiverComponent::FindOrCreateFFTOceanPatchSingleton(UWorld* World)
+{
+	AFFTOceanPatchActor* Actor = nullptr;
+	if (!World)
+	{
+		UE_LOG(LogShallowWater, Warning, TEXT("No valid World."));
+	}
+	else
+	{
+		const TActorIterator<AFFTOceanPatchActor> ActorIterator(World);
+		if (ActorIterator)
+		{
+			Actor = *ActorIterator;
+		}
+		else
+		{
+			FActorSpawnParameters SpawnInfo;
+			SpawnInfo.SpawnCollisionHandlingOverride = ESpawnActorCollisionHandlingMethod::AlwaysSpawn;
+			Actor = World->SpawnActor<AFFTOceanPatchActor>(SpawnInfo);
+			
+			Actor->PostLoad();
+		}
+
+		if (!Actor)
+		{
+			UE_LOG(LogShallowWater, Warning, TEXT("No singleton actor could be found or created."));
+		}
+	}
+	return TObjectPtr<AFFTOceanPatchActor>(Actor);
+}
+
 void UShallowWaterRiverComponent::PostLoad()
 {
 	Super::PostLoad();	
@@ -85,7 +118,7 @@ void UShallowWaterRiverComponent::PostLoad()
 		RiverSimSystem->Activate();
 	}
 
-	bRenderStateTickInitialize = false;		
+	bRenderStateTickInitialize = false;
 }
 
 void UShallowWaterRiverComponent::TickComponent(float DeltaTime, ELevelTick TickType, FActorComponentTickFunction* ThisTickFunction)
@@ -593,11 +626,17 @@ void UShallowWaterRiverComponent::Rebuild()
 		BakedWaterSurfaceNormalRT = NewObject<UTextureRenderTarget2D>(this);	
 		BakedWaterSurfaceNormalRT->InitAutoFormat(1, 1);
 		RiverSimSystem->SetVariableTextureRenderTarget(FName("NormalRT"), BakedWaterSurfaceNormalRT);
+	}
 
-		NormalDetailRT = NewObject<UTextureRenderTarget2D>(this);	
-		NormalDetailRT->InitAutoFormat(1, 1);
-		RiverSimSystem->SetVariableTextureRenderTarget(FName("NormalDetailRT"), NormalDetailRT);
+	TObjectPtr<AFFTOceanPatchActor> OceanPatchActor = FindOrCreateFFTOceanPatchSingleton(GetWorld());
+	if (OceanPatchActor == nullptr)
+	{
+		UE_LOG(LogShallowWater, Warning, TEXT("UShallowWaterRiverComponent::Rebuild() - No valid ocean patch actor can be found"));
+		return;
 	}
+		
+	NormalDetailRT = OceanPatchActor->OceanNormalRT;
+		RiverSimSystem->SetVariableTextureRenderTarget(FName("NormalDetailRT"), NormalDetailRT);
 
 	bIsInitialized = true;
 }
@@ -923,6 +962,14 @@ void UShallowWaterRiverComponent::OnWaterInfoTextureArrayCreated(const UTextureR
 
 void UShallowWaterRiverComponent::UpdateRenderState()
 {
+	TObjectPtr<AFFTOceanPatchActor> OceanPatchActor = FindOrCreateFFTOceanPatchSingleton(GetWorld());
+	if (OceanPatchActor == nullptr)
+	{
+		UE_LOG(LogShallowWater, Warning, TEXT("UShallowWaterRiverComponent::Rebuild() - No valid ocean patch actor can be found"));
+		return;
+	}
+	NormalDetailRT = OceanPatchActor->OceanNormalRT;
+
 	if (BakedSimMaterial == nullptr)
 	{
 		BakedSimMaterial = LoadObject<UMaterialInstance>(nullptr, TEXT("/WaterAdvanced/Niagara/Materials/SW_Water_Material_River.SW_Water_Material_River"));
@@ -966,6 +1013,9 @@ void UShallowWaterRiverComponent::UpdateRenderState()
 		RiverSimSystem->SetVariableBool(FName("DebugRenderBottomContour"), RenderState == EShallowWaterRenderState::DebugRenderBottomContour);
 		RiverSimSystem->SetVariableBool(FName("DebugRenderFoam"), RenderState == EShallowWaterRenderState::DebugRenderFoam);
 		RiverSimSystem->SetVariableBool(FName("ReadCachedSim"), bReadBakedSim);
+		
+		RiverSimSystem->SetVariableTextureRenderTarget("OceanNormalRT", NormalDetailRT);
+		RiverSimSystem->ResetSystem();
 	}
 
 	for (TSoftObjectPtr<AWaterBody > CurrWaterBody : AllWaterBodies)
@@ -1038,11 +1088,6 @@ void UShallowWaterRiverComponent::UpdateRenderState()
 	}
 
 	bRenderStateTickInitialize = true;
-
-	if (RiverSimSystem != nullptr)
-	{	
-		RiverSimSystem->ResetSystem();
-	}
 }
 
 void UShallowWaterRiverComponent::SetWaterMIDParameters(UMaterialInstanceDynamic* WaterMID)
diff --git a/Engine/Plugins/Experimental/WaterAdvanced/Source/WaterAdvanced/Public/FFTOceanPatchActor.h b/Engine/Plugins/Experimental/WaterAdvanced/Source/WaterAdvanced/Public/FFTOceanPatchActor.h
new file mode 100644
--- /dev/null
+++ b/Engine/Plugins/Experimental/WaterAdvanced/Source/WaterAdvanced/Public/FFTOceanPatchActor.h
@@ -0,0 +1,34 @@
+// Copyright Epic Games, Inc. All Rights Reserved.
+
+#pragma once
+
+#include "CoreMinimal.h"
+#include "Engine/EngineTypes.h"
+#include "GameFramework/Actor.h"
+#include "Components/PrimitiveComponent.h"
+#include "NiagaraSystem.h"
+#include "Engine/TextureRenderTarget2D.h"
+#include "FFTOceanPatchActor.generated.h"
+
+class UNiagaraComponent;
+class UNiagaraSystem;
+class AWaterBody;
+class AWaterZone;
+
+UCLASS(BlueprintType, HideCategories = (Physics, Replication, Input, Collision))
+class WATERADVANCED_API AFFTOceanPatchActor : public AActor
+{
+	GENERATED_UCLASS_BODY()
+
+public:
+	UPROPERTY(BlueprintReadOnly, VisibleDefaultsOnly, Category = "Ocean")
+	TObjectPtr <class UNiagaraSystem> NiagaraOceanSimulation;
+
+	UPROPERTY(BlueprintReadOnly, VisibleDefaultsOnly, Category = "Ocean")
+	TObjectPtr<UNiagaraComponent> FFTOceanSystem;
+
+	UPROPERTY(BlueprintReadOnly, VisibleDefaultsOnly, Category = "Ocean")
+	TObjectPtr<UTextureRenderTarget2D> OceanNormalRT;
+
+	virtual void PostLoad() override;
+};
\ No newline at end of file
diff --git a/Engine/Plugins/Experimental/WaterAdvanced/Source/WaterAdvanced/Public/ShallowWaterRiverActor.h b/Engine/Plugins/Experimental/WaterAdvanced/Source/WaterAdvanced/Public/ShallowWaterRiverActor.h
--- a/Engine/Plugins/Experimental/WaterAdvanced/Source/WaterAdvanced/Public/ShallowWaterRiverActor.h
+++ b/Engine/Plugins/Experimental/WaterAdvanced/Source/WaterAdvanced/Public/ShallowWaterRiverActor.h
@@ -8,14 +8,11 @@
 #include "Components/PrimitiveComponent.h"
 #include "BakedShallowWaterSimulationComponent.h"
 #include "Math/Float16Color.h"
-#include "CoreMinimal.h"
-#include "Engine/EngineTypes.h"
-#include "NiagaraBakerOutput.h"
-#include "NiagaraBakerSettings.h"
 #include "NiagaraSystem.h"
 #include "UObject/GCObject.h"
 #include "Engine/TextureRenderTarget2D.h"
 #include "Engine/TextureRenderTarget2DArray.h"
+#include "FFTOceanPatchActor.h"
 #include "ShallowWaterRiverActor.generated.h"
 
 class UNiagaraComponent;
@@ -41,6 +38,10 @@ class WATERADVANCED_API UShallowWaterRiverComponent : public UPrimitiveComponent
 	GENERATED_UCLASS_BODY()
 
 public:
+	// Find/create the level set renderer singleton actor as required. Return whether the found or created actor.
+	static TObjectPtr<AFFTOceanPatchActor> FindOrCreateFFTOceanPatchSingleton(UWorld* World);
+
+
 	UPROPERTY(EditAnywhere, Category = "Simulation", meta = (DisplayName = "Niagara River Simulation"))
 	TObjectPtr <class UNiagaraSystem> NiagaraRiverSimulation;
 
