small fix

diff --git a/Engine/Plugins/Experimental/WaterAdvanced/Source/WaterAdvanced/Private/FFTOceanPatchActor.cpp b/Engine/Plugins/Experimental/WaterAdvanced/Source/WaterAdvanced/Private/FFTOceanPatchActor.cpp
--- a/Engine/Plugins/Experimental/WaterAdvanced/Source/WaterAdvanced/Private/FFTOceanPatchActor.cpp
+++ b/Engine/Plugins/Experimental/WaterAdvanced/Source/WaterAdvanced/Private/FFTOceanPatchActor.cpp
@@ -39,6 +39,7 @@ AFFTOceanPatchActor::AFFTOceanPatchActor(const FObjectInitializer& ObjectInitial
 
 void AFFTOceanPatchActor::PostLoad()
 {
+	Super::PostLoad();
 
 	if (NiagaraOceanSimulation == nullptr)
 	{
