fix initialization of shared ocean actor

#virtualized

diff --git a/Engine/Plugins/Experimental/WaterAdvanced/Source/WaterAdvanced/Private/FFTOceanPatchActor.cpp b/Engine/Plugins/Experimental/WaterAdvanced/Source/WaterAdvanced/Private/FFTOceanPatchActor.cpp
--- a/Engine/Plugins/Experimental/WaterAdvanced/Source/WaterAdvanced/Private/FFTOceanPatchActor.cpp
+++ b/Engine/Plugins/Experimental/WaterAdvanced/Source/WaterAdvanced/Private/FFTOceanPatchActor.cpp
@@ -46,23 +46,31 @@ void AFFTOceanPatchActor::PostLoad()
 		NiagaraOceanSimulation = LoadObject<UNiagaraSystem>(nullptr, TEXT("/WaterAdvanced/Niagara/Systems/Grid2D_OceanPatch.Grid2D_OceanPatch"));
 	}
 
-	FFTOceanSystem = NewObject<UNiagaraComponent>(this, NAME_None, RF_Public);	
-	
-	if (GetWorld() && GetWorld()->bIsWorldInitialized)
+	if (FFTOceanSystem == nullptr)
 	{
-		if (!FFTOceanSystem->IsRegistered())
+		FFTOceanSystem = NewObject<UNiagaraComponent>(this, NAME_None, RF_Public);	
+	
+		if (GetWorld() && GetWorld()->bIsWorldInitialized)
 		{
-			FFTOceanSystem->RegisterComponentWithWorld(GetWorld());
+			if (!FFTOceanSystem->IsRegistered())
+			{
+				FFTOceanSystem->RegisterComponentWithWorld(GetWorld());
+			}
 		}
+	}
 
-		FFTOceanSystem->SetVisibleFlag(true);
-		FFTOceanSystem->SetAsset(NiagaraOceanSimulation);
+	Initialize();
+}
 
-		OceanNormalRT = NewObject<UTextureRenderTarget2D>(this, NAME_None, RF_Transient);
-		OceanNormalRT->InitAutoFormat(1, 1);
-		FFTOceanSystem->SetVariableTextureRenderTarget(FName("OceanNormalRT"), OceanNormalRT);
-		FFTOceanSystem->ReinitializeSystem();
-		FFTOceanSystem->ResetSystem();
-		FFTOceanSystem->Activate();
-	}
+void AFFTOceanPatchActor::Initialize()
+{		
+	FFTOceanSystem->SetVisibleFlag(true);
+	FFTOceanSystem->SetAsset(NiagaraOceanSimulation);
+
+	OceanNormalRT = NewObject<UTextureRenderTarget2D>(this, NAME_None, RF_Transient);
+	OceanNormalRT->InitAutoFormat(1, 1);
+	FFTOceanSystem->SetVariableTextureRenderTarget(FName("OceanNormalRT"), OceanNormalRT);
+	FFTOceanSystem->ReinitializeSystem();
+	FFTOceanSystem->ResetSystem();
+	FFTOceanSystem->Activate();
 }
diff --git a/Engine/Plugins/Experimental/WaterAdvanced/Source/WaterAdvanced/Private/ShallowWaterRiverActor.cpp b/Engine/Plugins/Experimental/WaterAdvanced/Source/WaterAdvanced/Private/ShallowWaterRiverActor.cpp
--- a/Engine/Plugins/Experimental/WaterAdvanced/Source/WaterAdvanced/Private/ShallowWaterRiverActor.cpp
+++ b/Engine/Plugins/Experimental/WaterAdvanced/Source/WaterAdvanced/Private/ShallowWaterRiverActor.cpp
@@ -635,8 +635,14 @@ void UShallowWaterRiverComponent::Rebuild()
 		return;
 	}
 		
+	if (OceanPatchActor->OceanNormalRT == nullptr)
+	{
+		UE_LOG(LogShallowWater, Warning, TEXT("UShallowWaterRiverComponent::Rebuild() - ocean patch actor isn't initialized"));
+		return;
+	}
+
 	NormalDetailRT = OceanPatchActor->OceanNormalRT;
-		RiverSimSystem->SetVariableTextureRenderTarget(FName("NormalDetailRT"), NormalDetailRT);
+	RiverSimSystem->SetVariableTextureRenderTarget(FName("NormalDetailRT"), NormalDetailRT);
 
 	bIsInitialized = true;
 }
@@ -965,9 +971,16 @@ void UShallowWaterRiverComponent::UpdateRenderState()
 	TObjectPtr<AFFTOceanPatchActor> OceanPatchActor = FindOrCreateFFTOceanPatchSingleton(GetWorld());
 	if (OceanPatchActor == nullptr)
 	{
-		UE_LOG(LogShallowWater, Warning, TEXT("UShallowWaterRiverComponent::Rebuild() - No valid ocean patch actor can be found"));
+		UE_LOG(LogShallowWater, Warning, TEXT("UShallowWaterRiverComponent::UpdateRenderState() - No valid ocean patch actor can be found"));
 		return;
 	}
+
+	if (OceanPatchActor->OceanNormalRT == nullptr)
+	{
+		UE_LOG(LogShallowWater, Warning, TEXT("UShallowWaterRiverComponent::UpdateRenderState() - ocean patch actor isn't initialized"));
+		return;
+	}
+
 	NormalDetailRT = OceanPatchActor->OceanNormalRT;
 
 	if (BakedSimMaterial == nullptr)
diff --git a/Engine/Plugins/Experimental/WaterAdvanced/Source/WaterAdvanced/Public/FFTOceanPatchActor.h b/Engine/Plugins/Experimental/WaterAdvanced/Source/WaterAdvanced/Public/FFTOceanPatchActor.h
--- a/Engine/Plugins/Experimental/WaterAdvanced/Source/WaterAdvanced/Public/FFTOceanPatchActor.h
+++ b/Engine/Plugins/Experimental/WaterAdvanced/Source/WaterAdvanced/Public/FFTOceanPatchActor.h
@@ -31,4 +31,6 @@ public:
 	TObjectPtr<UTextureRenderTarget2D> OceanNormalRT;
 
 	virtual void PostLoad() override;
+
+	void Initialize();
 };
\ No newline at end of file
