change warnings to verbose log messages to reduce log spam

diff --git a/Engine/Plugins/Experimental/WaterAdvanced/Source/WaterAdvanced/Private/ShallowWaterRiverActor.cpp b/Engine/Plugins/Experimental/WaterAdvanced/Source/WaterAdvanced/Private/ShallowWaterRiverActor.cpp
--- a/Engine/Plugins/Experimental/WaterAdvanced/Source/WaterAdvanced/Private/ShallowWaterRiverActor.cpp
+++ b/Engine/Plugins/Experimental/WaterAdvanced/Source/WaterAdvanced/Private/ShallowWaterRiverActor.cpp
@@ -263,7 +263,7 @@ void UShallowWaterRiverComponent::Rebuild()
 			}
 			else 
 			{
-				UE_LOG(LogShallowWater, Warning, TEXT("UShallowWaterRiverComponent::Rebuild() - skipping null water body actor found"));
+				UE_LOG(LogShallowWater, Verbose, TEXT("UShallowWaterRiverComponent::Rebuild() - skipping null water body actor found"));
 				continue;
 			}
 		}
@@ -1035,7 +1035,7 @@ void UShallowWaterRiverComponent::UpdateRenderState()
 	{
 		if (!CurrWaterBody)
 		{
-			UE_LOG(LogShallowWater, Warning, TEXT("UShallowWaterRiverComponent::UpdateRenderState() - Water Body Actor is null"));
+			UE_LOG(LogShallowWater, Verbose, TEXT("UShallowWaterRiverComponent::UpdateRenderState() - Water Body Actor is null"));
 			return;
 		}
 
