move FFT ocean actor from singleton to subsystem

diff --git a/Engine/Plugins/Experimental/WaterAdvanced/Source/WaterAdvanced/Private/FFTOceanPatchActor.cpp b/Engine/Plugins/Experimental/WaterAdvanced/Source/WaterAdvanced/Private/FFTOceanPatchActor.cpp
--- a/Engine/Plugins/Experimental/WaterAdvanced/Source/WaterAdvanced/Private/FFTOceanPatchActor.cpp
+++ b/Engine/Plugins/Experimental/WaterAdvanced/Source/WaterAdvanced/Private/FFTOceanPatchActor.cpp
@@ -32,7 +32,7 @@
 #include "LevelInstance/LevelInstanceSubsystem.h"
 
 AFFTOceanPatchActor::AFFTOceanPatchActor(const FObjectInitializer& ObjectInitializer)
-	: Super(ObjectInitializer)
+	: Super(ObjectInitializer), bIsInitialized(false)
 {	
 	SetHidden(false);
 }
@@ -40,37 +40,55 @@ AFFTOceanPatchActor::AFFTOceanPatchActor(const FObjectInitializer& ObjectInitial
 void AFFTOceanPatchActor::PostLoad()
 {
 	Super::PostLoad();
+}
+
+void AFFTOceanPatchActor::Initialize(UWorld *InWorld)
+{	
+	bIsInitialized = false;
 
 	if (NiagaraOceanSimulation == nullptr)
 	{
 		NiagaraOceanSimulation = LoadObject<UNiagaraSystem>(nullptr, TEXT("/WaterAdvanced/Niagara/Systems/Grid2D_OceanPatch.Grid2D_OceanPatch"));
 	}
+	else
+	{
+		UE_LOG(LogShallowWater, Warning, TEXT("AFFTOceanPatchActor::PostLoad - No valid FFT ocean patch system to use."));	
+	}
 
 	if (FFTOceanSystem == nullptr)
 	{
-		FFTOceanSystem = NewObject<UNiagaraComponent>(this, NAME_None, RF_Public);	
-	
-		if (GetWorld() && GetWorld()->bIsWorldInitialized)
+		FFTOceanSystem = UNiagaraFunctionLibrary::SpawnSystemAtLocation(InWorld, NiagaraOceanSimulation,
+			FVector(0,0,0), FRotator::ZeroRotator, FVector::OneVector, false,
+			false,
+			ENCPoolMethod::None, false);
+		
+		if (FFTOceanSystem == nullptr)
+		{
+			UE_LOG(LogShallowWater, Warning, TEXT("AFFTOceanPatchActor::PostLoad - Cannot spawn fft ocean system"));	
+			return;
+		}
+
+		if (InWorld && InWorld->bIsWorldInitialized)
 		{
 			if (!FFTOceanSystem->IsRegistered())
 			{
-				FFTOceanSystem->RegisterComponentWithWorld(GetWorld());
+				FFTOceanSystem->RegisterComponentWithWorld(InWorld);
 			}
 		}
+		else
+		{
+			UE_LOG(LogShallowWater, Warning, TEXT("AFFTOceanPatchActor::PostLoad - World not initialized"));	
+			return;
+		}
 	}
 
-	Initialize();
-}
-
-void AFFTOceanPatchActor::Initialize()
-{		
 	FFTOceanSystem->SetVisibleFlag(true);
 	FFTOceanSystem->SetAsset(NiagaraOceanSimulation);
 
 	OceanNormalRT = NewObject<UTextureRenderTarget2D>(this, NAME_None, RF_Transient);
 	OceanNormalRT->InitAutoFormat(1, 1);
 	FFTOceanSystem->SetVariableTextureRenderTarget(FName("OceanNormalRT"), OceanNormalRT);
-	FFTOceanSystem->ReinitializeSystem();
-	FFTOceanSystem->ResetSystem();
 	FFTOceanSystem->Activate();
+
+	bIsInitialized = true;
 }
diff --git a/Engine/Plugins/Experimental/WaterAdvanced/Source/WaterAdvanced/Private/FFTOceanPatchSubsystem.cpp b/Engine/Plugins/Experimental/WaterAdvanced/Source/WaterAdvanced/Private/FFTOceanPatchSubsystem.cpp
new file mode 100644
--- /dev/null
+++ b/Engine/Plugins/Experimental/WaterAdvanced/Source/WaterAdvanced/Private/FFTOceanPatchSubsystem.cpp
@@ -0,0 +1,71 @@
+// Copyright Epic Games, Inc. All Rights Reserved.
+
+#include "FFTOceanPatchSubsystem.h"
+#include "ShallowWaterRiverSettings.h"
+#include "Engine/World.h"
+
+UFFTOceanPatchSubsystem::UFFTOceanPatchSubsystem()
+{
+}
+
+bool UFFTOceanPatchSubsystem::ShouldCreateSubsystem(UObject* Outer) const
+{
+	TObjectPtr<UShallowWaterRiverSettings> TmpSettings = GetMutableDefault<UShallowWaterRiverSettings>();
+	if (TmpSettings && !TmpSettings->UseFFTPatch)
+	{
+		return false;
+	}
+
+		if (!FApp::CanEverRender() || IsRunningDedicatedServer())
+	{
+		return false;
+	}
+
+	// IsRunningDedicatedServer() is a static check and doesn't work in PIE "As Client' mode where both a server and a client are run
+	if (UWorld* World = Cast<UWorld>(Outer))
+	{
+		if (World->IsNetMode(NM_DedicatedServer))
+		{
+			return false;
+		}
+	}
+	
+	return Super::ShouldCreateSubsystem(Outer);	
+}
+
+void UFFTOceanPatchSubsystem::OnWorldBeginPlay(UWorld& InWorld)
+{
+	FFTOceanPatchActor = nullptr;
+}
+
+TStatId UFFTOceanPatchSubsystem::GetStatId() const
+{
+	RETURN_QUICK_DECLARE_CYCLE_STAT(UFFTOceanPatchSubsystem, STATGROUP_Tickables);
+}
+
+TObjectPtr<AFFTOceanPatchActor> UFFTOceanPatchSubsystem::GetFFTOceanPatchActor(UWorld *InWorld)
+{
+	if (FFTOceanPatchActor == nullptr)
+	{
+		FFTOceanPatchActor = NewObject<AFFTOceanPatchActor>((UObject*) InWorld, NAME_None, RF_Transient);
+
+		if (FFTOceanPatchActor != nullptr)
+		{
+			FFTOceanPatchActor->PostLoad();
+			FFTOceanPatchActor->Initialize(InWorld);
+
+			if (!FFTOceanPatchActor->IsInitialized())
+			{
+				FFTOceanPatchActor = nullptr;
+
+				return nullptr;
+			}
+		}
+		else
+		{
+			return nullptr;
+		}		
+	}
+
+	return FFTOceanPatchActor;
+}
diff --git a/Engine/Plugins/Experimental/WaterAdvanced/Source/WaterAdvanced/Private/ShallowWaterRiverActor.cpp b/Engine/Plugins/Experimental/WaterAdvanced/Source/WaterAdvanced/Private/ShallowWaterRiverActor.cpp
--- a/Engine/Plugins/Experimental/WaterAdvanced/Source/WaterAdvanced/Private/ShallowWaterRiverActor.cpp
+++ b/Engine/Plugins/Experimental/WaterAdvanced/Source/WaterAdvanced/Private/ShallowWaterRiverActor.cpp
@@ -23,6 +23,7 @@
 
 #include "TextureResource.h"
 #include "ShallowWaterCommon.h"
+#include "FFTOceanPatchSubsystem.h"
 #include "Engine/World.h"
 #include "Engine/Texture2D.h"
 #include "Math/Float16Color.h"
@@ -70,33 +71,33 @@ UShallowWaterRiverComponent::UShallowWaterRiverComponent(const FObjectInitialize
 
 TObjectPtr<AFFTOceanPatchActor> UShallowWaterRiverComponent::FindOrCreateFFTOceanPatchSingleton(UWorld* World)
 {
-	AFFTOceanPatchActor* Actor = nullptr;
-	if (!World)
+	if (World != nullptr)
 	{
-		UE_LOG(LogShallowWater, Warning, TEXT("No valid World."));
-	}
-	else
-	{
-		const TActorIterator<AFFTOceanPatchActor> ActorIterator(World);
-		if (ActorIterator)
+		UFFTOceanPatchSubsystem *OceanPatchSubsystem = World->GetSubsystem<UFFTOceanPatchSubsystem>();
+
+		if (OceanPatchSubsystem != nullptr)
 		{
-			Actor = *ActorIterator;
+			TObjectPtr<AFFTOceanPatchActor> OceanPatchActor = OceanPatchSubsystem->GetFFTOceanPatchActor(World);
+			if (OceanPatchActor != nullptr)
+			{
+				return OceanPatchActor;
+			}
+			else
+			{
+				UE_LOG(LogShallowWater, Warning, TEXT("No valid FFT ocean patch actor."));	
+			}
 		}
 		else
 		{
-			FActorSpawnParameters SpawnInfo;
-			SpawnInfo.SpawnCollisionHandlingOverride = ESpawnActorCollisionHandlingMethod::AlwaysSpawn;
-			Actor = World->SpawnActor<AFFTOceanPatchActor>(SpawnInfo);
-			
-			Actor->PostLoad();
-		}
-
-		if (!Actor)
-		{
-			UE_LOG(LogShallowWater, Warning, TEXT("No singleton actor could be found or created."));
+			UE_LOG(LogShallowWater, Warning, TEXT("No valid FFT ocean patch subsystem."));	
 		}
 	}
-	return TObjectPtr<AFFTOceanPatchActor>(Actor);
+	else
+	{
+		UE_LOG(LogShallowWater, Warning, TEXT("No valid World."));
+	}
+
+	return nullptr;
 }
 
 void UShallowWaterRiverComponent::PostLoad()
diff --git a/Engine/Plugins/Experimental/WaterAdvanced/Source/WaterAdvanced/Private/ShallowWaterRiverSettings.cpp b/Engine/Plugins/Experimental/WaterAdvanced/Source/WaterAdvanced/Private/ShallowWaterRiverSettings.cpp
new file mode 100644
--- /dev/null
+++ b/Engine/Plugins/Experimental/WaterAdvanced/Source/WaterAdvanced/Private/ShallowWaterRiverSettings.cpp
@@ -0,0 +1,12 @@
+// Copyright Epic Games, Inc. All Rights Reserved.
+
+#include "ShallowWaterRiverSettings.h"
+
+UShallowWaterRiverSettings::UShallowWaterRiverSettings()
+{
+}
+
+FName UShallowWaterRiverSettings::GetCategoryName() const
+{
+	return FName("Plugins");
+}
diff --git a/Engine/Plugins/Experimental/WaterAdvanced/Source/WaterAdvanced/Public/FFTOceanPatchActor.h b/Engine/Plugins/Experimental/WaterAdvanced/Source/WaterAdvanced/Public/FFTOceanPatchActor.h
--- a/Engine/Plugins/Experimental/WaterAdvanced/Source/WaterAdvanced/Public/FFTOceanPatchActor.h
+++ b/Engine/Plugins/Experimental/WaterAdvanced/Source/WaterAdvanced/Public/FFTOceanPatchActor.h
@@ -32,5 +32,9 @@ public:
 
 	virtual void PostLoad() override;
 
-	void Initialize();
+	void Initialize(UWorld *InWorld);
+	bool IsInitialized() { return bIsInitialized; }
+
+private:
+	bool bIsInitialized;
 };
\ No newline at end of file
diff --git a/Engine/Plugins/Experimental/WaterAdvanced/Source/WaterAdvanced/Public/FFTOceanPatchSubsystem.h b/Engine/Plugins/Experimental/WaterAdvanced/Source/WaterAdvanced/Public/FFTOceanPatchSubsystem.h
new file mode 100644
--- /dev/null
+++ b/Engine/Plugins/Experimental/WaterAdvanced/Source/WaterAdvanced/Public/FFTOceanPatchSubsystem.h
@@ -0,0 +1,24 @@
+// Copyright Epic Games, Inc. All Rights Reserved.
+
+#pragma once
+
+#include "Subsystems/WorldSubsystem.h"
+#include "FFTOceanPatchActor.h"
+#include "FFTOceanPatchSubsystem.generated.h"
+
+UCLASS(Blueprintable, Transient)
+class WATERADVANCED_API UFFTOceanPatchSubsystem : public UTickableWorldSubsystem
+{
+	GENERATED_BODY()
+
+public:
+	UFFTOceanPatchSubsystem();
+
+	virtual bool ShouldCreateSubsystem(UObject* Outer) const override;
+	virtual void OnWorldBeginPlay(UWorld& InWorld) override;	
+	virtual TStatId GetStatId() const override;
+
+	TObjectPtr<AFFTOceanPatchActor> GetFFTOceanPatchActor(UWorld *InWorld);	
+private:
+	TObjectPtr<AFFTOceanPatchActor> FFTOceanPatchActor;	
+};
diff --git a/Engine/Plugins/Experimental/WaterAdvanced/Source/WaterAdvanced/Public/ShallowWaterRiverSettings.h b/Engine/Plugins/Experimental/WaterAdvanced/Source/WaterAdvanced/Public/ShallowWaterRiverSettings.h
new file mode 100644
--- /dev/null
+++ b/Engine/Plugins/Experimental/WaterAdvanced/Source/WaterAdvanced/Public/ShallowWaterRiverSettings.h
@@ -0,0 +1,21 @@
+// Copyright Epic Games, Inc. All Rights Reserved.
+
+#pragma once
+
+#include "CoreMinimal.h"
+#include "Engine/DeveloperSettings.h"
+#include "Materials/MaterialParameterCollection.h"
+#include "ShallowWaterRiverSettings.generated.h"
+
+UCLASS(config = Engine, defaultconfig, meta=(DisplayName="Water Advanced"))
+class WATERADVANCED_API UShallowWaterRiverSettings : public UDeveloperSettings
+{
+	GENERATED_BODY()
+	UShallowWaterRiverSettings();
+	
+public:
+	virtual FName GetCategoryName() const override;	
+
+	UPROPERTY(Config, BlueprintReadOnly, EditAnywhere, Category = "Shallow Water River")
+	bool UseFFTPatch = true;
+};
