[Backout] Of previous commit

diff --git a/Engine/Plugins/Experimental/WaterAdvanced/Source/WaterAdvanced/Private/FFTOceanPatchActor.cpp b/Engine/Plugins/Experimental/WaterAdvanced/Source/WaterAdvanced/Private/FFTOceanPatchActor.cpp
--- a/Engine/Plugins/Experimental/WaterAdvanced/Source/WaterAdvanced/Private/FFTOceanPatchActor.cpp
+++ b/Engine/Plugins/Experimental/WaterAdvanced/Source/WaterAdvanced/Private/FFTOceanPatchActor.cpp
@@ -32,7 +32,7 @@
 #include "LevelInstance/LevelInstanceSubsystem.h"
 
 AFFTOceanPatchActor::AFFTOceanPatchActor(const FObjectInitializer& ObjectInitializer)
-	: Super(ObjectInitializer), bIsInitialized(false)
+	: Super(ObjectInitializer)
 {	
 	SetHidden(false);
 }
@@ -40,55 +40,37 @@ AFFTOceanPatchActor::AFFTOceanPatchActor(const FObjectInitializer& ObjectInitial
 void AFFTOceanPatchActor::PostLoad()
 {
 	Super::PostLoad();
-}
-
-void AFFTOceanPatchActor::Initialize(UWorld *InWorld)
-{	
-	bIsInitialized = false;
 
 	if (NiagaraOceanSimulation == nullptr)
 	{
 		NiagaraOceanSimulation = LoadObject<UNiagaraSystem>(nullptr, TEXT("/WaterAdvanced/Niagara/Systems/Grid2D_OceanPatch.Grid2D_OceanPatch"));
 	}
-	else
-	{
-		UE_LOG(LogShallowWater, Warning, TEXT("AFFTOceanPatchActor::PostLoad - No valid FFT ocean patch system to use."));	
-	}
 
 	if (FFTOceanSystem == nullptr)
 	{
-		FFTOceanSystem = UNiagaraFunctionLibrary::SpawnSystemAtLocation(InWorld, NiagaraOceanSimulation,
-			FVector(0,0,0), FRotator::ZeroRotator, FVector::OneVector, false,
-			false,
-			ENCPoolMethod::None, false);
-		
-		if (FFTOceanSystem == nullptr)
-		{
-			UE_LOG(LogShallowWater, Warning, TEXT("AFFTOceanPatchActor::PostLoad - Cannot spawn fft ocean system"));	
-			return;
-		}
-
-		if (InWorld && InWorld->bIsWorldInitialized)
+		FFTOceanSystem = NewObject<UNiagaraComponent>(this, NAME_None, RF_Public);	
+	
+		if (GetWorld() && GetWorld()->bIsWorldInitialized)
 		{
 			if (!FFTOceanSystem->IsRegistered())
 			{
-				FFTOceanSystem->RegisterComponentWithWorld(InWorld);
+				FFTOceanSystem->RegisterComponentWithWorld(GetWorld());
 			}
 		}
-		else
-		{
-			UE_LOG(LogShallowWater, Warning, TEXT("AFFTOceanPatchActor::PostLoad - World not initialized"));	
-			return;
-		}
 	}
 
+	Initialize();
+}
+
+void AFFTOceanPatchActor::Initialize()
+{		
 	FFTOceanSystem->SetVisibleFlag(true);
 	FFTOceanSystem->SetAsset(NiagaraOceanSimulation);
 
 	OceanNormalRT = NewObject<UTextureRenderTarget2D>(this, NAME_None, RF_Transient);
 	OceanNormalRT->InitAutoFormat(1, 1);
 	FFTOceanSystem->SetVariableTextureRenderTarget(FName("OceanNormalRT"), OceanNormalRT);
+	FFTOceanSystem->ReinitializeSystem();
+	FFTOceanSystem->ResetSystem();
 	FFTOceanSystem->Activate();
-
-	bIsInitialized = true;
 }
diff --git a/Engine/Plugins/Experimental/WaterAdvanced/Source/WaterAdvanced/Private/FFTOceanPatchSubsystem.cpp b/Engine/Plugins/Experimental/WaterAdvanced/Source/WaterAdvanced/Private/FFTOceanPatchSubsystem.cpp
deleted file mode 100644
diff --git a/Engine/Plugins/Experimental/WaterAdvanced/Source/WaterAdvanced/Private/ShallowWaterRiverActor.cpp b/Engine/Plugins/Experimental/WaterAdvanced/Source/WaterAdvanced/Private/ShallowWaterRiverActor.cpp
--- a/Engine/Plugins/Experimental/WaterAdvanced/Source/WaterAdvanced/Private/ShallowWaterRiverActor.cpp
+++ b/Engine/Plugins/Experimental/WaterAdvanced/Source/WaterAdvanced/Private/ShallowWaterRiverActor.cpp
@@ -23,7 +23,6 @@
 
 #include "TextureResource.h"
 #include "ShallowWaterCommon.h"
-#include "FFTOceanPatchSubsystem.h"
 #include "Engine/World.h"
 #include "Engine/Texture2D.h"
 #include "Math/Float16Color.h"
@@ -71,33 +70,33 @@ UShallowWaterRiverComponent::UShallowWaterRiverComponent(const FObjectInitialize
 
 TObjectPtr<AFFTOceanPatchActor> UShallowWaterRiverComponent::FindOrCreateFFTOceanPatchSingleton(UWorld* World)
 {
-	if (World != nullptr)
+	AFFTOceanPatchActor* Actor = nullptr;
+	if (!World)
 	{
-		UFFTOceanPatchSubsystem *OceanPatchSubsystem = World->GetSubsystem<UFFTOceanPatchSubsystem>();
-
-		if (OceanPatchSubsystem != nullptr)
+		UE_LOG(LogShallowWater, Warning, TEXT("No valid World."));
+	}
+	else
+	{
+		const TActorIterator<AFFTOceanPatchActor> ActorIterator(World);
+		if (ActorIterator)
 		{
-			TObjectPtr<AFFTOceanPatchActor> OceanPatchActor = OceanPatchSubsystem->GetFFTOceanPatchActor(World);
-			if (OceanPatchActor != nullptr)
-			{
-				return OceanPatchActor;
-			}
-			else
-			{
-				UE_LOG(LogShallowWater, Warning, TEXT("No valid FFT ocean patch actor."));	
-			}
+			Actor = *ActorIterator;
 		}
 		else
 		{
-			UE_LOG(LogShallowWater, Warning, TEXT("No valid FFT ocean patch subsystem."));	
+			FActorSpawnParameters SpawnInfo;
+			SpawnInfo.SpawnCollisionHandlingOverride = ESpawnActorCollisionHandlingMethod::AlwaysSpawn;
+			Actor = World->SpawnActor<AFFTOceanPatchActor>(SpawnInfo);
+			
+			Actor->PostLoad();
 		}
-	}
-	else
-	{
-		UE_LOG(LogShallowWater, Warning, TEXT("No valid World."));
-	}
 
-	return nullptr;
+		if (!Actor)
+		{
+			UE_LOG(LogShallowWater, Warning, TEXT("No singleton actor could be found or created."));
+		}
+	}
+	return TObjectPtr<AFFTOceanPatchActor>(Actor);
 }
 
 void UShallowWaterRiverComponent::PostLoad()
diff --git a/Engine/Plugins/Experimental/WaterAdvanced/Source/WaterAdvanced/Private/ShallowWaterRiverSettings.cpp b/Engine/Plugins/Experimental/WaterAdvanced/Source/WaterAdvanced/Private/ShallowWaterRiverSettings.cpp
deleted file mode 100644
diff --git a/Engine/Plugins/Experimental/WaterAdvanced/Source/WaterAdvanced/Public/FFTOceanPatchActor.h b/Engine/Plugins/Experimental/WaterAdvanced/Source/WaterAdvanced/Public/FFTOceanPatchActor.h
--- a/Engine/Plugins/Experimental/WaterAdvanced/Source/WaterAdvanced/Public/FFTOceanPatchActor.h
+++ b/Engine/Plugins/Experimental/WaterAdvanced/Source/WaterAdvanced/Public/FFTOceanPatchActor.h
@@ -32,9 +32,5 @@ public:
 
 	virtual void PostLoad() override;
 
-	void Initialize(UWorld *InWorld);
-	bool IsInitialized() { return bIsInitialized; }
-
-private:
-	bool bIsInitialized;
+	void Initialize();
 };
\ No newline at end of file
diff --git a/Engine/Plugins/Experimental/WaterAdvanced/Source/WaterAdvanced/Public/FFTOceanPatchSubsystem.h b/Engine/Plugins/Experimental/WaterAdvanced/Source/WaterAdvanced/Public/FFTOceanPatchSubsystem.h
deleted file mode 100644
diff --git a/Engine/Plugins/Experimental/WaterAdvanced/Source/WaterAdvanced/Public/ShallowWaterRiverSettings.h b/Engine/Plugins/Experimental/WaterAdvanced/Source/WaterAdvanced/Public/ShallowWaterRiverSettings.h
deleted file mode 100644
