force load NS with FObjectFinder

diff --git a/Engine/Plugins/Experimental/WaterAdvanced/Source/WaterAdvanced/Private/FFTOceanPatchActor.cpp b/Engine/Plugins/Experimental/WaterAdvanced/Source/WaterAdvanced/Private/FFTOceanPatchActor.cpp
--- a/Engine/Plugins/Experimental/WaterAdvanced/Source/WaterAdvanced/Private/FFTOceanPatchActor.cpp
+++ b/Engine/Plugins/Experimental/WaterAdvanced/Source/WaterAdvanced/Private/FFTOceanPatchActor.cpp
@@ -20,6 +20,7 @@
 #include "PhysicsEngine/ConvexElem.h"
 #include "PhysicsEngine/BodySetup.h"
 #include "Engine/OverlapResult.h"
+#include "UObject/ConstructorHelpers.h"
 
 #include "TextureResource.h"
 #include "ShallowWaterCommon.h"
@@ -35,6 +36,9 @@ AFFTOceanPatchActor::AFFTOceanPatchActor(const FObjectInitializer& ObjectInitial
 	: Super(ObjectInitializer), bIsInitialized(false)
 {	
 	SetHidden(false);
+
+	ConstructorHelpers::FObjectFinder<UNiagaraSystem> DefaultOceanPatchNiagaraSystem(TEXT("/WaterAdvanced/Niagara/Systems/Grid2D_OceanPatch.Grid2D_OceanPatch"));
+	NiagaraOceanSimulation = DefaultOceanPatchNiagaraSystem.Object;
 }
 
 void AFFTOceanPatchActor::PostLoad()
@@ -45,16 +49,7 @@ void AFFTOceanPatchActor::PostLoad()
 void AFFTOceanPatchActor::Initialize(UWorld *InWorld)
 {	
 	bIsInitialized = false;
-
-	if (NiagaraOceanSimulation == nullptr)
-	{
-		NiagaraOceanSimulation = LoadObject<UNiagaraSystem>(nullptr, TEXT("/WaterAdvanced/Niagara/Systems/Grid2D_OceanPatch.Grid2D_OceanPatch"));
-	}
-	else
-	{
-		UE_LOG(LogShallowWater, Warning, TEXT("AFFTOceanPatchActor::PostLoad - No valid FFT ocean patch system to use."));	
-	}
-
+	
 	if (FFTOceanSystem == nullptr)
 	{
 		FFTOceanSystem = UNiagaraFunctionLibrary::SpawnSystemAtLocation(InWorld, NiagaraOceanSimulation,
