move Niagara System loading to be for non CDO objects only

diff --git a/Engine/Plugins/Experimental/WaterAdvanced/Source/WaterAdvanced/Private/FFTOceanPatchActor.cpp b/Engine/Plugins/Experimental/WaterAdvanced/Source/WaterAdvanced/Private/FFTOceanPatchActor.cpp
--- a/Engine/Plugins/Experimental/WaterAdvanced/Source/WaterAdvanced/Private/FFTOceanPatchActor.cpp
+++ b/Engine/Plugins/Experimental/WaterAdvanced/Source/WaterAdvanced/Private/FFTOceanPatchActor.cpp
@@ -37,8 +37,11 @@ AFFTOceanPatchActor::AFFTOceanPatchActor(const FObjectInitializer& ObjectInitial
 {	
 	SetHidden(false);
 
-	ConstructorHelpers::FObjectFinder<UNiagaraSystem> DefaultOceanPatchNiagaraSystem(TEXT("/WaterAdvanced/Niagara/Systems/Grid2D_OceanPatch.Grid2D_OceanPatch"));
-	NiagaraOceanSimulation = DefaultOceanPatchNiagaraSystem.Object;
+	if (!HasAnyFlags(RF_ClassDefaultObject))
+	{
+		ConstructorHelpers::FObjectFinder<UNiagaraSystem> DefaultOceanPatchNiagaraSystem(TEXT("/WaterAdvanced/Niagara/Systems/Grid2D_OceanPatch.Grid2D_OceanPatch"));
+		NiagaraOceanSimulation = DefaultOceanPatchNiagaraSystem.Object;
+	}
 }
 
 void AFFTOceanPatchActor::PostLoad()
@@ -50,6 +53,11 @@ void AFFTOceanPatchActor::Initialize(UWorld *InWorld)
 {	
 	bIsInitialized = false;
 	
+	if (NiagaraOceanSimulation)
+	{
+
+	}
+
 	if (FFTOceanSystem == nullptr)
 	{
 		FFTOceanSystem = UNiagaraFunctionLibrary::SpawnSystemAtLocation(InWorld, NiagaraOceanSimulation,
