- async load fft Niagara System
- move logic from ocean patch actor to Subsystem
- remove ocean patch actor

#virtualized

diff --git a/Engine/Plugins/Experimental/WaterAdvanced/Source/WaterAdvanced/Private/FFTOceanPatchActor.cpp b/Engine/Plugins/Experimental/WaterAdvanced/Source/WaterAdvanced/Private/FFTOceanPatchActor.cpp
deleted file mode 100644
diff --git a/Engine/Plugins/Experimental/WaterAdvanced/Source/WaterAdvanced/Private/FFTOceanPatchSubsystem.cpp b/Engine/Plugins/Experimental/WaterAdvanced/Source/WaterAdvanced/Private/FFTOceanPatchSubsystem.cpp
--- a/Engine/Plugins/Experimental/WaterAdvanced/Source/WaterAdvanced/Private/FFTOceanPatchSubsystem.cpp
+++ b/Engine/Plugins/Experimental/WaterAdvanced/Source/WaterAdvanced/Private/FFTOceanPatchSubsystem.cpp
@@ -1,8 +1,14 @@
 // Copyright Epic Games, Inc. All Rights Reserved.
 
 #include "FFTOceanPatchSubsystem.h"
-#include "ShallowWaterRiverSettings.h"
+#include "ShallowWaterSettings.h"
 #include "Engine/World.h"
+#include "Engine/AssetManager.h"
+#include "NiagaraComponent.h"
+#include "NiagaraFunctionLibrary.h"
+#include "ShallowWaterCommon.h"
+#include "NiagaraComponent.h"
+#include "Engine/TextureRenderTarget2D.h"
 
 UFFTOceanPatchSubsystem::UFFTOceanPatchSubsystem()
 {
@@ -10,13 +16,7 @@ UFFTOceanPatchSubsystem::UFFTOceanPatchSubsystem()
 
 bool UFFTOceanPatchSubsystem::ShouldCreateSubsystem(UObject* Outer) const
 {
-	TObjectPtr<UShallowWaterRiverSettings> TmpSettings = GetMutableDefault<UShallowWaterRiverSettings>();
-	if (TmpSettings && !TmpSettings->UseFFTPatch)
-	{
-		return false;
-	}
-
-		if (!FApp::CanEverRender() || IsRunningDedicatedServer())
+	if (!FApp::CanEverRender() || IsRunningDedicatedServer())
 	{
 		return false;
 	}
@@ -33,9 +33,30 @@ bool UFFTOceanPatchSubsystem::ShouldCreateSubsystem(UObject* Outer) const
 	return Super::ShouldCreateSubsystem(Outer);	
 }
 
-void UFFTOceanPatchSubsystem::OnWorldBeginPlay(UWorld& InWorld)
+void UFFTOceanPatchSubsystem::PostInitialize()
 {
-	FFTOceanPatchActor = nullptr;
+	Super::PostInitialize();
+
+	Settings = GetMutableDefault<UShallowWaterSettings>();
+
+	if (!Settings)
+	{
+		ensureMsgf(false, TEXT("UShallowWaterSubsystem::PostInitialize() - UShallowWaterSettings is not valid"));
+		return;
+	}
+
+	TArray< FSoftObjectPath> ObjectsToLoad;
+	ObjectsToLoad.Add(Settings->DefaultOceanPatchNiagaraSystem.ToSoftObjectPath());
+
+	UAssetManager::GetStreamableManager().RequestAsyncLoad(ObjectsToLoad,
+	FStreamableDelegate::CreateWeakLambda(this, [this]()
+		{
+			// continue with initialization		
+		})
+	);	
+
+	FFTOceanSystem = nullptr;
+	OceanNormalRT = nullptr;
 }
 
 TStatId UFFTOceanPatchSubsystem::GetStatId() const
@@ -43,29 +64,51 @@ TStatId UFFTOceanPatchSubsystem::GetStatId() const
 	RETURN_QUICK_DECLARE_CYCLE_STAT(UFFTOceanPatchSubsystem, STATGROUP_Tickables);
 }
 
-TObjectPtr<AFFTOceanPatchActor> UFFTOceanPatchSubsystem::GetFFTOceanPatchActor(UWorld *InWorld)
+TObjectPtr<UTextureRenderTarget2D> UFFTOceanPatchSubsystem::GetOceanNormalRT(UWorld* World)
 {
-	if (FFTOceanPatchActor == nullptr)
+	if (FFTOceanSystem == nullptr)
 	{
-		FFTOceanPatchActor = NewObject<AFFTOceanPatchActor>((UObject*) InWorld, NAME_None, RF_Transient);
+		Settings = GetMutableDefault<UShallowWaterSettings>();		
+		UNiagaraSystem *NiagaraOceanSimulation = Settings->DefaultOceanPatchNiagaraSystem.Get();
 
-		if (FFTOceanPatchActor != nullptr)
+		if (NiagaraOceanSimulation == nullptr)
 		{
-			FFTOceanPatchActor->PostLoad();
-			FFTOceanPatchActor->Initialize(InWorld);
+			UE_LOG(LogShallowWater, Warning, TEXT("UFFTOceanPatchSubsystem::GetOceanNormalRT - Ocean simulation system not loaded"));	
+			return nullptr;
+		}
 
-			if (!FFTOceanPatchActor->IsInitialized())
-			{
-				FFTOceanPatchActor = nullptr;
+		FFTOceanSystem = UNiagaraFunctionLibrary::SpawnSystemAtLocation(World, NiagaraOceanSimulation,
+			FVector(0,0,0), FRotator::ZeroRotator, FVector::OneVector, false,
+			false,
+			ENCPoolMethod::None, false);
+		
+		if (FFTOceanSystem == nullptr)
+		{
+			UE_LOG(LogShallowWater, Warning, TEXT("UFFTOceanPatchSubsystem::GetOceanNormalRT - Cannot spawn fft ocean system"));	
+			return nullptr;
+		}
 
-				return nullptr;
+		if (World && World->bIsWorldInitialized)
+		{
+			if (!FFTOceanSystem->IsRegistered())
+			{
+				FFTOceanSystem->RegisterComponentWithWorld(World);
 			}
 		}
 		else
 		{
+			UE_LOG(LogShallowWater, Warning, TEXT("UFFTOceanPatchSubsystem::GetOceanNormalRT - World not initialized"));	
 			return nullptr;
-		}		
+		}
+
+		FFTOceanSystem->SetVisibleFlag(true);
+		FFTOceanSystem->SetAsset(NiagaraOceanSimulation);
+
+		OceanNormalRT = NewObject<UTextureRenderTarget2D>(this, NAME_None, RF_Transient);
+		OceanNormalRT->InitAutoFormat(1, 1);
+		FFTOceanSystem->SetVariableTextureRenderTarget(FName("OceanNormalRT"), OceanNormalRT);
+		FFTOceanSystem->Activate();		
 	}
 
-	return FFTOceanPatchActor;
-}
+	return OceanNormalRT;
+}
\ No newline at end of file
diff --git a/Engine/Plugins/Experimental/WaterAdvanced/Source/WaterAdvanced/Private/ShallowWaterRiverActor.cpp b/Engine/Plugins/Experimental/WaterAdvanced/Source/WaterAdvanced/Private/ShallowWaterRiverActor.cpp
--- a/Engine/Plugins/Experimental/WaterAdvanced/Source/WaterAdvanced/Private/ShallowWaterRiverActor.cpp
+++ b/Engine/Plugins/Experimental/WaterAdvanced/Source/WaterAdvanced/Private/ShallowWaterRiverActor.cpp
@@ -69,7 +69,7 @@ UShallowWaterRiverComponent::UShallowWaterRiverComponent(const FObjectInitialize
 
 }
 
-TObjectPtr<AFFTOceanPatchActor> UShallowWaterRiverComponent::FindOrCreateFFTOceanPatchSingleton(UWorld* World)
+TObjectPtr<UTextureRenderTarget2D> UShallowWaterRiverComponent::GetSharedFFTOceanPatchNormalRTFromSubsystem(UWorld* World)
 {
 	if (World != nullptr)
 	{
@@ -77,15 +77,7 @@ TObjectPtr<AFFTOceanPatchActor> UShallowWaterRiverComponent::FindOrCreateFFTOcea
 
 		if (OceanPatchSubsystem != nullptr)
 		{
-			TObjectPtr<AFFTOceanPatchActor> OceanPatchActor = OceanPatchSubsystem->GetFFTOceanPatchActor(World);
-			if (OceanPatchActor != nullptr)
-			{
-				return OceanPatchActor;
-			}
-			else
-			{
-				UE_LOG(LogShallowWater, Warning, TEXT("No valid FFT ocean patch actor."));	
-			}
+			return OceanPatchSubsystem->GetOceanNormalRT(World);
 		}
 		else
 		{
@@ -629,20 +621,15 @@ void UShallowWaterRiverComponent::Rebuild()
 		RiverSimSystem->SetVariableTextureRenderTarget(FName("NormalRT"), BakedWaterSurfaceNormalRT);
 	}
 
-	TObjectPtr<AFFTOceanPatchActor> OceanPatchActor = FindOrCreateFFTOceanPatchSingleton(GetWorld());
-	if (OceanPatchActor == nullptr)
-	{
-		UE_LOG(LogShallowWater, Warning, TEXT("UShallowWaterRiverComponent::Rebuild() - No valid ocean patch actor can be found"));
-		return;
-	}
-		
-	if (OceanPatchActor->OceanNormalRT == nullptr)
+	TObjectPtr<UTextureRenderTarget2D>  OceanPatchNormalRT = GetSharedFFTOceanPatchNormalRTFromSubsystem(GetWorld());
+
+	if (OceanPatchNormalRT == nullptr)
 	{
-		UE_LOG(LogShallowWater, Warning, TEXT("UShallowWaterRiverComponent::Rebuild() - ocean patch actor isn't initialized"));
+		UE_LOG(LogShallowWater, Warning, TEXT("UShallowWaterRiverComponent::Rebuild() - ocean patch normal RT is not initialized"));
 		return;
 	}
 
-	NormalDetailRT = OceanPatchActor->OceanNormalRT;
+	NormalDetailRT = OceanPatchNormalRT;
 	RiverSimSystem->SetVariableTextureRenderTarget(FName("NormalDetailRT"), NormalDetailRT);
 
 	bIsInitialized = true;
@@ -969,20 +956,15 @@ void UShallowWaterRiverComponent::OnWaterInfoTextureArrayCreated(const UTextureR
 
 void UShallowWaterRiverComponent::UpdateRenderState()
 {
-	TObjectPtr<AFFTOceanPatchActor> OceanPatchActor = FindOrCreateFFTOceanPatchSingleton(GetWorld());
-	if (OceanPatchActor == nullptr)
-	{
-		UE_LOG(LogShallowWater, Warning, TEXT("UShallowWaterRiverComponent::UpdateRenderState() - No valid ocean patch actor can be found"));
-		return;
-	}
+	TObjectPtr<UTextureRenderTarget2D>  OceanPatchNormalRT = GetSharedFFTOceanPatchNormalRTFromSubsystem(GetWorld());
 
-	if (OceanPatchActor->OceanNormalRT == nullptr)
+	if (OceanPatchNormalRT == nullptr)
 	{
-		UE_LOG(LogShallowWater, Warning, TEXT("UShallowWaterRiverComponent::UpdateRenderState() - ocean patch actor isn't initialized"));
+		UE_LOG(LogShallowWater, Warning, TEXT("UShallowWaterRiverComponent::Rebuild() - ocean patch normal RT is not initialized"));
 		return;
 	}
-
-	NormalDetailRT = OceanPatchActor->OceanNormalRT;
+	
+	NormalDetailRT = OceanPatchNormalRT;
 
 	if (BakedSimMaterial == nullptr)
 	{
diff --git a/Engine/Plugins/Experimental/WaterAdvanced/Source/WaterAdvanced/Private/ShallowWaterRiverSettings.cpp b/Engine/Plugins/Experimental/WaterAdvanced/Source/WaterAdvanced/Private/ShallowWaterRiverSettings.cpp
deleted file mode 100644
diff --git a/Engine/Plugins/Experimental/WaterAdvanced/Source/WaterAdvanced/Private/ShallowWaterSettings.cpp b/Engine/Plugins/Experimental/WaterAdvanced/Source/WaterAdvanced/Private/ShallowWaterSettings.cpp
--- a/Engine/Plugins/Experimental/WaterAdvanced/Source/WaterAdvanced/Private/ShallowWaterSettings.cpp
+++ b/Engine/Plugins/Experimental/WaterAdvanced/Source/WaterAdvanced/Private/ShallowWaterSettings.cpp
@@ -6,7 +6,8 @@ UShallowWaterSettings::UShallowWaterSettings()
 {
 	DefaultShallowWaterNiagaraSimulation = TSoftObjectPtr<UNiagaraSystem>(FSoftObjectPath(TEXT("/WaterAdvanced/Niagara/Systems/Grid2D_SW_WaterBody.Grid2D_SW_WaterBody")));
 	DefaultShallowWaterCollisionNDC = TSoftObjectPtr<UNiagaraDataChannelAsset>(FSoftObjectPath(TEXT("/WaterAdvanced/Niagara/DataChannels/NDC_ShallowWater.NDC_ShallowWater")));
-	WaterMPC = TSoftObjectPtr<UMaterialParameterCollection>(FSoftObjectPath(TEXT("/Water/Materials/MPC/MPC_Water.MPC_Water")));
+	DefaultOceanPatchNiagaraSystem = TSoftObjectPtr<UNiagaraSystem>(FSoftObjectPath(TEXT("/WaterAdvanced/Niagara/Systems/Grid2D_OceanPatch.Grid2D_OceanPatch")));
+	WaterMPC = TSoftObjectPtr<UMaterialParameterCollection>(FSoftObjectPath(TEXT("/Water/Materials/MPC/MPC_Water.MPC_Water")));	
 }
 
 FName UShallowWaterSettings::GetCategoryName() const
diff --git a/Engine/Plugins/Experimental/WaterAdvanced/Source/WaterAdvanced/Public/FFTOceanPatchActor.h b/Engine/Plugins/Experimental/WaterAdvanced/Source/WaterAdvanced/Public/FFTOceanPatchActor.h
deleted file mode 100644
diff --git a/Engine/Plugins/Experimental/WaterAdvanced/Source/WaterAdvanced/Public/FFTOceanPatchSubsystem.h b/Engine/Plugins/Experimental/WaterAdvanced/Source/WaterAdvanced/Public/FFTOceanPatchSubsystem.h
--- a/Engine/Plugins/Experimental/WaterAdvanced/Source/WaterAdvanced/Public/FFTOceanPatchSubsystem.h
+++ b/Engine/Plugins/Experimental/WaterAdvanced/Source/WaterAdvanced/Public/FFTOceanPatchSubsystem.h
@@ -3,9 +3,12 @@
 #pragma once
 
 #include "Subsystems/WorldSubsystem.h"
-#include "FFTOceanPatchActor.h"
+#include "ShallowWaterSettings.h"
 #include "FFTOceanPatchSubsystem.generated.h"
 
+class UNiagaraComponent;
+class UTextureRenderTarget2D;
+
 UCLASS(Blueprintable, Transient)
 class WATERADVANCED_API UFFTOceanPatchSubsystem : public UTickableWorldSubsystem
 {
@@ -15,10 +18,14 @@ public:
 	UFFTOceanPatchSubsystem();
 
 	virtual bool ShouldCreateSubsystem(UObject* Outer) const override;
-	virtual void OnWorldBeginPlay(UWorld& InWorld) override;	
+	virtual void PostInitialize() override;
 	virtual TStatId GetStatId() const override;
-
-	TObjectPtr<AFFTOceanPatchActor> GetFFTOceanPatchActor(UWorld *InWorld);	
+	
+	TObjectPtr<UTextureRenderTarget2D> GetOceanNormalRT(UWorld* World);
+protected:
+	UPROPERTY(BlueprintReadOnly, VisibleAnywhere, Category="Shallow Water")
+	TObjectPtr<UShallowWaterSettings> Settings;
 private:
-	TObjectPtr<AFFTOceanPatchActor> FFTOceanPatchActor;	
+	TObjectPtr<UNiagaraComponent> FFTOceanSystem;
+	TObjectPtr<UTextureRenderTarget2D> OceanNormalRT;
 };
diff --git a/Engine/Plugins/Experimental/WaterAdvanced/Source/WaterAdvanced/Public/ShallowWaterRiverActor.h b/Engine/Plugins/Experimental/WaterAdvanced/Source/WaterAdvanced/Public/ShallowWaterRiverActor.h
--- a/Engine/Plugins/Experimental/WaterAdvanced/Source/WaterAdvanced/Public/ShallowWaterRiverActor.h
+++ b/Engine/Plugins/Experimental/WaterAdvanced/Source/WaterAdvanced/Public/ShallowWaterRiverActor.h
@@ -12,7 +12,6 @@
 #include "UObject/GCObject.h"
 #include "Engine/TextureRenderTarget2D.h"
 #include "Engine/TextureRenderTarget2DArray.h"
-#include "FFTOceanPatchActor.h"
 #include "ShallowWaterRiverActor.generated.h"
 
 class UNiagaraComponent;
@@ -38,9 +37,6 @@ class WATERADVANCED_API UShallowWaterRiverComponent : public UPrimitiveComponent
 	GENERATED_UCLASS_BODY()
 
 public:
-	// Find/create the level set renderer singleton actor as required. Return whether the found or created actor.
-	static TObjectPtr<AFFTOceanPatchActor> FindOrCreateFFTOceanPatchSingleton(UWorld* World);
-
 
 	UPROPERTY(EditAnywhere, Category = "Simulation", meta = (DisplayName = "Niagara River Simulation"))
 	TObjectPtr <class UNiagaraSystem> NiagaraRiverSimulation;
@@ -196,6 +192,9 @@ private:
 
 	UPROPERTY()
 	TObjectPtr<UTextureRenderTarget2D> NormalDetailRT;
+
+	// Find/create the level set renderer singleton actor as required. Return whether the found or created actor.
+	TObjectPtr<UTextureRenderTarget2D> GetSharedFFTOceanPatchNormalRTFromSubsystem(UWorld* World);
 };
 
 UCLASS(BlueprintType, HideCategories = (Physics, Replication, Input, Collision))
diff --git a/Engine/Plugins/Experimental/WaterAdvanced/Source/WaterAdvanced/Public/ShallowWaterRiverSettings.h b/Engine/Plugins/Experimental/WaterAdvanced/Source/WaterAdvanced/Public/ShallowWaterRiverSettings.h
deleted file mode 100644
diff --git a/Engine/Plugins/Experimental/WaterAdvanced/Source/WaterAdvanced/Public/ShallowWaterSettings.h b/Engine/Plugins/Experimental/WaterAdvanced/Source/WaterAdvanced/Public/ShallowWaterSettings.h
--- a/Engine/Plugins/Experimental/WaterAdvanced/Source/WaterAdvanced/Public/ShallowWaterSettings.h
+++ b/Engine/Plugins/Experimental/WaterAdvanced/Source/WaterAdvanced/Public/ShallowWaterSettings.h
@@ -24,6 +24,8 @@ public:
 	TSoftObjectPtr<UNiagaraSystem> DefaultShallowWaterNiagaraSimulation;
 	UPROPERTY(Config, BlueprintReadOnly, EditAnywhere, Category = "Shallow Water Simulation")
 	TSoftObjectPtr<UNiagaraDataChannelAsset> DefaultShallowWaterCollisionNDC;
+	UPROPERTY(Config, BlueprintReadOnly, EditAnywhere, Category="Shallow Water River")
+	TSoftObjectPtr<UNiagaraSystem> DefaultOceanPatchNiagaraSystem;
 	UPROPERTY(Config, BlueprintReadOnly, EditAnywhere, Category="Shallow Water Simulation")
 	FShallowWaterSimParameters ShallowWaterSimParameters;
 	UPROPERTY(Config, BlueprintReadOnly, EditAnywhere, Category="Shallow Water Simulation")
