ripples work with simulated pawns

diff --git a/Engine/Plugins/Experimental/WaterAdvanced/Source/WaterAdvanced/Private/ShallowWaterSubsystem.cpp b/Engine/Plugins/Experimental/WaterAdvanced/Source/WaterAdvanced/Private/ShallowWaterSubsystem.cpp
--- a/Engine/Plugins/Experimental/WaterAdvanced/Source/WaterAdvanced/Private/ShallowWaterSubsystem.cpp
+++ b/Engine/Plugins/Experimental/WaterAdvanced/Source/WaterAdvanced/Private/ShallowWaterSubsystem.cpp
@@ -523,7 +523,19 @@ void UShallowWaterSubsystem::UpdateGridMovement()
 	{
 		return;
 	}
-	const FVector CursorPawnLocation = CursorPawn->GetActorLocation();
+	
+	FVector CursorPawnLocation;
+
+	// if we have a simulated skel mesh then use the component location instead of the actor which is always at the last non simulated position
+	USkeletalMeshComponent *SkelMeshComponent = CursorPawn->GetComponentByClass<USkeletalMeshComponent>();
+	if (SkelMeshComponent != nullptr && SkelMeshComponent->BodyInstance.bSimulatePhysics)
+	{
+		CursorPawnLocation = SkelMeshComponent->GetComponentLocation();
+	}
+	else
+	{
+		CursorPawnLocation = CursorPawn->GetActorLocation();
+	}
 	
 	/*
 	 * Snap the sim center to player character
