faster iteration over pawns
remove check for skel mesh

diff --git a/Engine/Plugins/Experimental/WaterAdvanced/Source/WaterAdvanced/Private/ShallowWaterSubsystem.cpp b/Engine/Plugins/Experimental/WaterAdvanced/Source/WaterAdvanced/Private/ShallowWaterSubsystem.cpp
--- a/Engine/Plugins/Experimental/WaterAdvanced/Source/WaterAdvanced/Private/ShallowWaterSubsystem.cpp
+++ b/Engine/Plugins/Experimental/WaterAdvanced/Source/WaterAdvanced/Private/ShallowWaterSubsystem.cpp
@@ -26,6 +26,7 @@
 #include "PhysicsEngine/PhysicsAsset.h"
 #include "Engine/AssetManager.h"
 #include "Engine/Texture2D.h"
+#include "EngineUtils.h"
 
 //UE_DISABLE_OPTIMIZATION_SHIP
 
@@ -524,18 +525,7 @@ void UShallowWaterSubsystem::UpdateGridMovement()
 		return;
 	}
 	
-	FVector CursorPawnLocation;
-
-	// if we have a simulated skel mesh then use the component location instead of the actor which is always at the last non simulated position
-	USkeletalMeshComponent *SkelMeshComponent = CursorPawn->GetComponentByClass<USkeletalMeshComponent>();
-	if (SkelMeshComponent != nullptr && SkelMeshComponent->BodyInstance.bSimulatePhysics)
-	{
-		CursorPawnLocation = SkelMeshComponent->GetComponentLocation();
-	}
-	else
-	{
-		CursorPawnLocation = CursorPawn->GetActorLocation();
-	}
+	const FVector CursorPawnLocation = CursorPawn->GetActorLocation();
 	
 	/*
 	 * Snap the sim center to player character
@@ -1208,15 +1198,16 @@ TArray<APawn*> UShallowWaterSubsystem::GetPawnsInRange(const bool bShouldSortByS
 
 TArray<APawn*> UShallowWaterSubsystem::GetPawnsInRange(const FVector ObservingLocation, const bool bShouldSortBySignificance) const
 {
-	TArray<AActor*> FoundActors;
-	UGameplayStatics::GetAllActorsOfClass(Cast<UObject>(GetWorld()), APawn::StaticClass(), FoundActors);
+	UWorld* World = GetWorld();
+
 	TArray<APawn*> Results;
-	for (AActor* Actor : FoundActors)
+	for (TActorIterator<APawn> It(World, APawn::StaticClass()); It; ++It)
 	{
-		const float Distance = FVector::Distance(ObservingLocation, Actor->GetActorLocation());
+		APawn *Pawn = *It;
+		const float Distance = FVector::Distance(ObservingLocation, Pawn->GetActorLocation());
 		if (Distance <= GetColliderMaxRange())
 		{
-			Results.Add(Cast<APawn>(Actor));
+			Results.Add(Pawn);
 		}
 	}
 
