support cases where we remove water bodies from the simulated list

diff --git a/Engine/Plugins/Experimental/WaterAdvanced/Source/WaterAdvanced/Private/ShallowWaterRiverActor.cpp b/Engine/Plugins/Experimental/WaterAdvanced/Source/WaterAdvanced/Private/ShallowWaterRiverActor.cpp
--- a/Engine/Plugins/Experimental/WaterAdvanced/Source/WaterAdvanced/Private/ShallowWaterRiverActor.cpp
+++ b/Engine/Plugins/Experimental/WaterAdvanced/Source/WaterAdvanced/Private/ShallowWaterRiverActor.cpp
@@ -747,9 +747,25 @@ void UShallowWaterRiverComponent::Bake()
 	BakedFoamTexture = BakedFoamRT->ConstructTexture2D(this, "BakedFoamTexture", TextureObjectFlags);
 	RiverSimSystem->SetVariableTexture(FName("BakedFoamTexture"), BakedFoamTexture);
 	
+	// clear references to old baked sim on water body actors
+	if (BakedSim != nullptr)
+	{ 
+		for (TSoftObjectPtr<AWaterBody > CurrWaterBody : BakedSim->WaterBodies)
+		{
+			TObjectPtr<UWaterBodyComponent> CurrWaterBodyComponent = CurrWaterBody->GetWaterBodyComponent();
+
+			if (CurrWaterBodyComponent != nullptr)
+			{
+				CurrWaterBodyComponent->SetBakedShallowWaterSimulation(nullptr);
+				CurrWaterBodyComponent->PostEditChange();
+			}
+		}
+	}
+
 	BakedSim = NewObject<UBakedShallowWaterSimulationComponent>(this, NAME_None, RF_Public);
 	BakedSim->SimulationData = FShallowWaterSimulationGrid(ShallowWaterSimArrayValues, BakedWaterSurfaceTexture, FIntVector2(BakedWaterSurfaceRT->SizeX, BakedWaterSurfaceRT->SizeY), SystemPos, WorldGridSize);
-		
+	BakedSim->WaterBodies = AllWaterBodies;	
+
 	// compute the maximum water height for each convex in each water body simulated by this river
 	// we use this to modify the collision geometry so it fully encompasses the baked water sim
 	TMap<FKConvexElem*, float> ConvexToMaxHeight;
