fix issues with unloaded water bodies

diff --git a/Engine/Plugins/Experimental/WaterAdvanced/Source/WaterAdvanced/Private/ShallowWaterRiverActor.cpp b/Engine/Plugins/Experimental/WaterAdvanced/Source/WaterAdvanced/Private/ShallowWaterRiverActor.cpp
--- a/Engine/Plugins/Experimental/WaterAdvanced/Source/WaterAdvanced/Private/ShallowWaterRiverActor.cpp
+++ b/Engine/Plugins/Experimental/WaterAdvanced/Source/WaterAdvanced/Private/ShallowWaterRiverActor.cpp
@@ -714,6 +714,18 @@ void UShallowWaterRiverComponent::Bake()
 		return;
 	}
 
+	if (SourceRiverWaterBodies.Num() != 0)
+	{
+		for (TSoftObjectPtr<AWaterBody > CurrWaterBody : SourceRiverWaterBodies)
+		{
+			if (CurrWaterBody == nullptr)
+			{			
+				UE_LOG(LogShallowWater, Warning, TEXT("UShallowWaterRiverComponent::Rebuild() - Cannot bake with a null water body.  Please make sure all water bodies are loaded and that all entires on the actor are valid"));
+				return;
+			}
+		}
+	}
+
 	BakedWaterSurfaceTexture = BakedWaterSurfaceRT->ConstructTexture2D(this, "BakedRiverTexture", TextureObjectFlags);
 
 	RiverSimSystem->SetVariableTexture(FName("BakedSimTexture"), BakedWaterSurfaceTexture);
@@ -1027,15 +1039,15 @@ void UShallowWaterRiverComponent::UpdateRenderState()
 		RiverSimSystem->SetVariableBool(FName("ReadCachedSim"), bReadBakedSim);
 		
 		RiverSimSystem->SetVariableTextureRenderTarget("OceanNormalRT", NormalDetailRT);
-		RiverSimSystem->ResetSystem();
+		RiverSimSystem->ReinitializeSystem();
 	}
 
 	for (TSoftObjectPtr<AWaterBody > CurrWaterBody : AllWaterBodies)
 	{
 		if (!CurrWaterBody)
 		{
-			UE_LOG(LogShallowWater, Verbose, TEXT("UShallowWaterRiverComponent::UpdateRenderState() - Water Body Actor is null"));
-			return;
+			UE_LOG(LogShallowWater, Warning, TEXT("UShallowWaterRiverComponent::UpdateRenderState() - Water Body Actor is null- skipping setting render state"));
+			continue;
 		}
 
 		TObjectPtr<UWaterBodyRiverComponent> CurrWaterBodyComponent = Cast<UWaterBodyRiverComponent>(CurrWaterBody->GetWaterBodyComponent());
