double precision grid location
lwc fix for water sim position
custom bilinear interp for baked water texture

#virtualized

diff --git a/Engine/Plugins/Experimental/WaterAdvanced/Source/WaterAdvanced/Private/ShallowWaterRiverActor.cpp b/Engine/Plugins/Experimental/WaterAdvanced/Source/WaterAdvanced/Private/ShallowWaterRiverActor.cpp
--- a/Engine/Plugins/Experimental/WaterAdvanced/Source/WaterAdvanced/Private/ShallowWaterRiverActor.cpp
+++ b/Engine/Plugins/Experimental/WaterAdvanced/Source/WaterAdvanced/Private/ShallowWaterRiverActor.cpp
@@ -777,7 +777,7 @@ void UShallowWaterRiverComponent::Bake()
 	BakedSim = NewObject<UBakedShallowWaterSimulationComponent>(this, NAME_None, RF_Public);
 	BakedSim->SimulationData = FShallowWaterSimulationGrid(ShallowWaterSimArrayValues, BakedWaterSurfaceTexture, FIntVector2(BakedWaterSurfaceRT->SizeX, BakedWaterSurfaceRT->SizeY), SystemPos, WorldGridSize);
 	BakedSim->WaterBodies = AllWaterBodies;	
-
+		
 	// compute the maximum water height for each convex in each water body simulated by this river
 	// we use this to modify the collision geometry so it fully encompasses the baked water sim
 	TMap<FKConvexElem*, float> ConvexToMaxHeight;
@@ -1076,6 +1076,7 @@ void UShallowWaterRiverComponent::UpdateRenderState()
 					WaterInfoMID->SetTextureParameterValue("BakedWaterSimTex", BakedWaterSurfaceTexture);
 					WaterInfoMID->SetTextureParameterValue("FoamTex", BakedFoamTexture);
 					WaterInfoMID->SetVectorParameterValue("BakedWaterSimLocation", SystemPos);
+					WaterMID->SetDoubleVectorParameterValue("BakedWaterSimLocationDouble", SystemPos);
 					WaterInfoMID->SetVectorParameterValue("BakedWaterSimSize", FVector(WorldGridSize.X, WorldGridSize.Y, 1));
 				}
 				else
@@ -1122,6 +1123,7 @@ void UShallowWaterRiverComponent::SetWaterMIDParameters(UMaterialInstanceDynamic
 		WaterMID->SetTextureParameterValue("FoamTex", BakedFoamTexture);
 
 		WaterMID->SetVectorParameterValue("BakedWaterSimLocation", SystemPos);
+		WaterMID->SetDoubleVectorParameterValue("BakedWaterSimLocationDouble", SystemPos);
 		WaterMID->SetVectorParameterValue("BakedWaterSimSize", FVector(WorldGridSize.X, WorldGridSize.Y, 1));
 
 		WaterMID->SetTextureParameterValue("BakedWaterSimNormalTex", BakedWaterSurfaceNormalRT);
