Optimize ShallowWaterSubsystem to avoid duplicate actor queries and weak pointer resolves during tick.
Clean up header to only set UPROPERTY and USTRUCT for hard pointers that need exposure to GC.
Small API change to some protected functions that were unlikely to be overridden.

diff --git a/Engine/Plugins/Experimental/WaterAdvanced/Source/WaterAdvanced/Private/ShallowWaterSubsystem.cpp b/Engine/Plugins/Experimental/WaterAdvanced/Source/WaterAdvanced/Private/ShallowWaterSubsystem.cpp
--- a/Engine/Plugins/Experimental/WaterAdvanced/Source/WaterAdvanced/Private/ShallowWaterSubsystem.cpp
+++ b/Engine/Plugins/Experimental/WaterAdvanced/Source/WaterAdvanced/Private/ShallowWaterSubsystem.cpp
@@ -82,22 +82,17 @@ static FAutoConsoleVariableRef CVarSWUUseFullVehiclePhysicsAssetse(
 
 FName UShallowWaterSubsystem::ColliderComponentTag = FName("RigidMesh_ShallowWaterCollider");
 
-TArray<AWaterBody*> FShallowWaterCollisionTracker_Actor::GetOverlappingWaterBodies() const
+void FShallowWaterCollisionTracker_Actor::GetOverlappingWaterBodies(TSet<AWaterBody*>& WaterBodies) const
 {
-	TArray<AWaterBody*> Result;
 	if(const AActor* ActorPointer = CollisionActor.Get())
 	{
-    	TArray<AActor*> OverlappingActors;
+    	TSet<AActor*> OverlappingActors;
 		ActorPointer->GetOverlappingActors(OverlappingActors, AWaterBody::StaticClass());
 		for (AActor* WaterBody : OverlappingActors)
 		{
-			if (AWaterBody* Actor = Cast<AWaterBody>(WaterBody))
-			{
-				Result.AddUnique(Actor);
-			}
+			WaterBodies.Add(CastChecked<AWaterBody>(WaterBody));
 		}
 	}
-    return Result;
 }
 
 UShallowWaterSubsystem::UShallowWaterSubsystem()
@@ -173,7 +168,6 @@ void UShallowWaterSubsystem::Tick(float DeltaTime)
 	}
 
 	UpdateActivePawns();
-	GatherContextsFromPawns(ActivePawns);
 	UpdateCollisionForPendingContexts();
 
 	/*
@@ -207,9 +201,9 @@ void UShallowWaterSubsystem::Tick(float DeltaTime)
 
 			// loop over the water bodies for the current water zone and make sure each one has all of the materials setup with 
 			// references to the simulation texture(s)
-			if (const UWaterBodyComponent* WaterBodyComp = Water->GetWaterBodyComponent())
+			if (UWaterBodyComponent* WaterBodyComp = Water->GetWaterBodyComponent())
 			{
-				if (const AWaterZone* WaterZone = WaterBodyComp->GetWaterZone())
+				if (AWaterZone* WaterZone = WaterBodyComp->GetWaterZone())
 				{
 					// Update all water bodies in the water zone only if it hasn't already been done
 					if (!WaterZonesWithMaterialsInitialized.Find(WaterZone))
@@ -255,6 +249,8 @@ void UShallowWaterSubsystem::Tick(float DeltaTime)
 			ShallowWaterNiagaraSimulation->Deactivate();
 		}
 	}
+
+	ClearTickCache();
 }
 
 TStatId UShallowWaterSubsystem::GetStatId() const
@@ -432,6 +428,11 @@ TOptional<FVector> UShallowWaterSubsystem::GetCameraLocationFromWeakController()
 
 APawn* UShallowWaterSubsystem::GetTheMostRelevantPlayerPawn() const
 {
+	if (bTickCacheValid && CachedCursorPawn)
+	{
+		return CachedCursorPawn;
+	}
+
 	const APlayerController* Controller = WeakPlayerController.Get();
 	if (Controller)
 	{
@@ -744,8 +745,8 @@ void UShallowWaterSubsystem::RegisterPhysicsAssetProxiesDataAsset(const UShallow
 TSet<AWaterBody*> UShallowWaterSubsystem::GetAllOverlappingWaterBodiesAndUpdateCollisionTrackers() 
 {
 	TSet<AWaterBody*> Result = GetOverlappingWaterBodiesFromPawns();
-	Result.Append(GetOverlappingWaterBodiesFromActorTrackersAndUpdate());
-	Result.Append(GetOverlappingWaterBodiesFromDirectTrackersAndUpdate());
+	GetOverlappingWaterBodiesFromActorTrackersAndUpdate(Result);
+	GetOverlappingWaterBodiesFromDirectTrackersAndUpdate(Result);
 	return Result;
 }
 
@@ -792,24 +793,23 @@ TSet<AWaterBody*> UShallowWaterSubsystem::GetOverlappingWaterBodiesFromPawns() c
 {
 	TSet<AWaterBody*> Result;
 	TArray<APawn*> Pawns = GetPawnsInRange();
+
+	TSet<AActor*> OverlappingActors;
+	TSubclassOf<AActor> WaterBodyClass = AWaterBody::StaticClass();
+
 	for (const APawn* Pawn : Pawns)
 	{
-		TArray<AActor*> OverlappingActors;
-		Pawn->GetOverlappingActors(OverlappingActors, AWaterBody::StaticClass());
+		Pawn->GetOverlappingActors(OverlappingActors, WaterBodyClass);
 		for (AActor* WaterBody : OverlappingActors)
 		{
-			if (AWaterBody* Actor = Cast<AWaterBody>(WaterBody))
-			{
-				Result.Add(Actor);
-			}
+			Result.Add(CastChecked<AWaterBody>(WaterBody));
 		}
 	}
 	return Result;
 }
 
-TSet<AWaterBody*> UShallowWaterSubsystem::GetOverlappingWaterBodiesFromActorTrackersAndUpdate()
+void UShallowWaterSubsystem::GetOverlappingWaterBodiesFromActorTrackersAndUpdate(TSet<AWaterBody*>& WaterBodies)
 {
-	TSet<AWaterBody*> Result;
 	for (int32 Index = Tracker_Actors.Num() - 1; Index >= 0; Index--)
 	{
 		if (!Tracker_Actors[Index].IsValid(GetWorld()->GetTimeSeconds()))
@@ -817,14 +817,12 @@ TSet<AWaterBody*> UShallowWaterSubsystem::GetOverlappingWaterBodiesFromActorTrac
 			Tracker_Actors.RemoveAt(Index);
 			continue;
 		}
-		Result.Append(Tracker_Actors[Index].GetOverlappingWaterBodies());
+		Tracker_Actors[Index].GetOverlappingWaterBodies(WaterBodies);
 	}
-	return Result;
 }
 
-TSet<AWaterBody*> UShallowWaterSubsystem::GetOverlappingWaterBodiesFromDirectTrackersAndUpdate()
+void UShallowWaterSubsystem::GetOverlappingWaterBodiesFromDirectTrackersAndUpdate(TSet<AWaterBody*>& WaterBodies)
 {
-	TSet<AWaterBody*> Result;
 	for (int32 Index = Tracker_Directs.Num() - 1; Index >= 0; Index--)
 	{
 		if (!Tracker_Directs[Index].IsValid(GetWorld()->GetTimeSeconds()))
@@ -832,12 +830,11 @@ TSet<AWaterBody*> UShallowWaterSubsystem::GetOverlappingWaterBodiesFromDirectTra
 			Tracker_Directs.RemoveAt(Index);
 			continue;
 		}
-		if(AWaterBody* WaterBody = Tracker_Directs[Index].GetOverlappingWaterBody())
+		if (AWaterBody* WaterBody = Tracker_Directs[Index].GetOverlappingWaterBody())
 		{
-			Result.Add(WaterBody);
+			WaterBodies.Add(WaterBody);
 		}
 	}
-	return Result;
 }
 
 void UShallowWaterSubsystem::UpdateOverlappingWaterBodiesHistory(TArray<AWaterBody*> OverlappingWaterBodies)
@@ -858,47 +855,60 @@ bool UShallowWaterSubsystem::ShouldSimulateThisFrame() const
 	return (GetWorld()->GetTimeSeconds() - LastTimeOverlappingAnyWaterBody <= GSWRemainActiveForSeconds);
 }
 
+void UShallowWaterSubsystem::ClearTickCache()
+{
+	bTickCacheValid = false;
+	CachedCursorPawn = nullptr;
+	CachedPawnsInRange.Reset();
+}
+
 int32 UShallowWaterSubsystem::UpdateActivePawns()
 {
-	const APawn* CursorPawn = GetTheMostRelevantPlayerPawn();
-	if (CursorPawn == nullptr)
+	bTickCacheValid = false;
+
+	CachedCursorPawn = GetTheMostRelevantPlayerPawn();
+	if (CachedCursorPawn == nullptr)
 	{
 		ActivePawns.Empty();
 		return 0;
 	}
-	FVector CursorPawnLocation = CursorPawn->GetActorLocation();
+
+	FVector CursorPawnLocation = CachedCursorPawn->GetActorLocation();
+	
+	// Use an array of hard pointers to speed up processing
+	TArray<APawn*> ValidActivePawns;
+	ValidActivePawns.Reserve(ActivePawns.Num());
+
+	FVector::FReal DistanceSquared = FMath::Square(GetColliderMaxRange());
 	
-	// Remove irrelevant pawns from ActivePawns
-	ActivePawns.RemoveAll([&](TWeakObjectPtr<APawn> Pawn)
+	ActivePawns.RemoveAll([&](TWeakObjectPtr<APawn> WeakPawn)
 	{
-		if(Pawn == nullptr)
+		APawn* Pawn = WeakPawn.Get();
+		if (Pawn && (Pawn->GetActorLocation() - CursorPawnLocation).SquaredLength() <= DistanceSquared)
 		{
-			return true;
-		}
-		if ((Pawn->GetActorLocation() - CursorPawnLocation).SquaredLength() > GetColliderMaxRange() * GetColliderMaxRange())
-		{
-			return true;
+			// Add to valid list if it's within range
+			ValidActivePawns.Add(Pawn);
+			return false;
 		}
-		return false;
+		// Remove from weak pointer list
+		return true;
 	});
 
-	const int32 SlotsLeft = Settings->MaxActivePawnNum - ActivePawns.Num();
+	// Update the cached pawns here for later
+	CachedPawnsInRange = GetPawnsInRange(CursorPawnLocation, false);
+
+	const int32 SlotsLeft = Settings->MaxActivePawnNum - ValidActivePawns.Num();
 	int32 NewPawnsAdded = 0;
 	if (SlotsLeft > 0)
 	{
 		// Get pawns nearby who become relevant this frame
-		TArray<APawn*> RelevantPawns = GetPawnsInRange();
-		RelevantPawns.RemoveAll([&](APawn* Pawn)
+		TArray<APawn*> RelevantPawns = CachedPawnsInRange;
+				
+		// Remove ones already in active list
+		RelevantPawns.RemoveAllSwap([&](APawn* Pawn)
 		{
-			for (const TWeakObjectPtr<APawn> Other : ActivePawns)
-			{
-				if (Pawn == nullptr || Pawn == Other.Get())
-				{
-					return true;
-				}
-			}
-			return false;
-		});
+			return ValidActivePawns.Contains(Pawn);
+		}, EAllowShrinking::No);
 
 		// Sort pawn candidates
 		RelevantPawns.Sort([CursorPawnLocation] (const APawn& Left, const APawn& Right)
@@ -912,7 +922,9 @@ int32 UShallowWaterSubsystem::UpdateActivePawns()
 		NewPawnsAdded = FMath::Min(SlotsLeft, RelevantPawns.Num());
 		for (int32 Index = 0; Index < NewPawnsAdded; Index++)
 		{
+			// Update both weak and strong lists
 			ActivePawns.Add(RelevantPawns[Index]);
+			ValidActivePawns.Add(RelevantPawns[Index]);
 		}
 
 #if ENABLE_DRAW_DEBUG
@@ -922,24 +934,22 @@ int32 UShallowWaterSubsystem::UpdateActivePawns()
 		}
 #endif
 	}
-	
-	return NewPawnsAdded;
-}
 
-void UShallowWaterSubsystem::GatherContextsFromPawns(TArray<TWeakObjectPtr<APawn>> InActivePawns)
-{
-	for (const TWeakObjectPtr<APawn>& Pawn : InActivePawns)
+	// Update pending collision contexts for new list
+	for (APawn* PawnPtr : ValidActivePawns)
 	{
-		if (APawn* PawnPtr = Pawn.Get())
+		TOptional<FShallowWaterCollisionContext> Context = GetCollisionContextFromPawn(PawnPtr);
+		if (Context.IsSet())
 		{
-			TOptional<FShallowWaterCollisionContext> Context = GetCollisionContextFromPawn(PawnPtr);
-			if (Context.IsSet())
-			{
-				// AddUnique() since multiple pawns can return the same Context. E.g. a multi-seat vehicle
-				PendingContexts.AddUnique(Context.GetValue());
-			}
+			// AddUnique() since multiple pawns can return the same Context. E.g. a multi-seat vehicle
+			PendingContexts.AddUnique(Context.GetValue());
 		}
 	}
+	
+	// Update tick cache, ResetTickCache must be called before the calling function returns to ensure safety
+	bTickCacheValid = true;
+
+	return NewPawnsAdded;
 }
 
 TOptional<FShallowWaterCollisionContext> UShallowWaterSubsystem::GetCollisionContextFromPawn(APawn* InPawn) const
@@ -953,12 +963,7 @@ TOptional<FShallowWaterCollisionContext> UShallowWaterSubsystem::GetCollisionCon
 	}
 	else
 	{
-		TArray<USkeletalMeshComponent*> SKMs;
-		InPawn->GetComponents<USkeletalMeshComponent>(SKMs);
-		if (!SKMs.IsEmpty())
-		{
-			Component = SKMs[0];
-		}
+		Component = InPawn->GetComponentByClass<USkeletalMeshComponent>();
 	}
 	if (Component)
 	{
@@ -991,32 +996,25 @@ void UShallowWaterSubsystem::UpdateCollisionForPendingContexts()
 	CleanUpVehicleCollisionProxies();
 	
 	// Find components that need their collision disabled
-	TArray<FShallowWaterCollisionContext> ContextsToDisable = PreviousContexts;
-	ContextsToDisable.RemoveAll([&] (FShallowWaterCollisionContext Context)
+	for (const FShallowWaterCollisionContext& Context : PreviousContexts)
 	{
-		return PendingContexts.Find(Context) != INDEX_NONE;
-	});
+		if (!PendingContexts.Contains(Context))
+		{
+			DisableCollisionForContext(Context);
+		}
+	}
 
 	// Find components that don't already exist
-	TArray<FShallowWaterCollisionContext> ContextsToEnable = PendingContexts;
-	ContextsToEnable.RemoveAll([&] (FShallowWaterCollisionContext Context)
-	{
-		return PreviousContexts.Find(Context) != INDEX_NONE;
-	});
-	
-
-	
-	for (const FShallowWaterCollisionContext& Context : ContextsToDisable)
+	for (const FShallowWaterCollisionContext& Context : PendingContexts)
 	{
-		DisableCollisionForContext(Context);
-	}
-	for (const FShallowWaterCollisionContext& Context : ContextsToEnable)
-	{
-		EnableCollisionForContext(Context);
+		if (!PreviousContexts.Contains(Context))
+		{
+			EnableCollisionForContext(Context);
+		}
 	}
 
 	PreviousContexts = PendingContexts;
-	PendingContexts.Empty();
+	PendingContexts.Reset();
 }
 
 void UShallowWaterSubsystem::EnableCollisionForContext(const FShallowWaterCollisionContext& Context)
@@ -1167,6 +1165,11 @@ float UShallowWaterSubsystem::GetColliderMaxRange() const
 
 TArray<APawn*> UShallowWaterSubsystem::GetPawnsInRange(const bool bShouldSortBySignificance) const
 {
+	if (bTickCacheValid && !bShouldSortBySignificance)
+	{
+		return CachedPawnsInRange;
+	}
+
 	const APlayerController* Controller = WeakPlayerController.Get();
 	APawn* CursorPawn = nullptr;
 	if (Controller)
@@ -1200,12 +1203,12 @@ TArray<APawn*> UShallowWaterSubsystem::GetPawnsInRange(const FVector ObservingLo
 {
 	UWorld* World = GetWorld();
 
+	FVector::FReal DistanceSquared = FMath::Square(GetColliderMaxRange());
 	TArray<APawn*> Results;
 	for (TActorIterator<APawn> It(World, APawn::StaticClass()); It; ++It)
 	{
 		APawn *Pawn = *It;
-		const float Distance = FVector::Distance(ObservingLocation, Pawn->GetActorLocation());
-		if (Distance <= GetColliderMaxRange())
+		if ((ObservingLocation - Pawn->GetActorLocation()).SquaredLength() <= DistanceSquared)
 		{
 			Results.Add(Pawn);
 		}
@@ -1215,7 +1218,7 @@ TArray<APawn*> UShallowWaterSubsystem::GetPawnsInRange(const FVector ObservingLo
 	{
 		Results.Sort([ObservingLocation](APawn& PawnA, APawn& PawnB)
 		{
-			return FVector::DistSquared(ObservingLocation, PawnA.GetActorLocation()) < FVector::DistSquared(ObservingLocation, PawnB.GetActorLocation());
+			return (ObservingLocation - PawnA.GetActorLocation()).SquaredLength() < (ObservingLocation - PawnB.GetActorLocation()).SquaredLength();
 		});
 	}
 	
diff --git a/Engine/Plugins/Experimental/WaterAdvanced/Source/WaterAdvanced/Public/ShallowWaterSubsystem.h b/Engine/Plugins/Experimental/WaterAdvanced/Source/WaterAdvanced/Public/ShallowWaterSubsystem.h
--- a/Engine/Plugins/Experimental/WaterAdvanced/Source/WaterAdvanced/Public/ShallowWaterSubsystem.h
+++ b/Engine/Plugins/Experimental/WaterAdvanced/Source/WaterAdvanced/Public/ShallowWaterSubsystem.h
@@ -30,7 +30,7 @@ enum class EShallowWaterCollisionContextType : uint8
 };
 
 // Shallow Water Rigid Body collision
-USTRUCT(BlueprintType)
+USTRUCT()
 struct WATERADVANCED_API FShallowWaterCollisionContext
 {
 	GENERATED_BODY()
@@ -70,22 +70,18 @@ struct WATERADVANCED_API FShallowWaterCollisionContext
  * CollisionTracker that records the actor affecting the waterbody, used by logs or fishing lures etc.,
  * where the collision is handled by Niagara Data Channel and the subsystem isn't aware of them without a tracking device. This is intended as a solution.
  */
-USTRUCT(BlueprintType)
 struct WATERADVANCED_API FShallowWaterCollisionTracker_Actor
 {
-	GENERATED_BODY()
-
 	FShallowWaterCollisionTracker_Actor() {}
 	FShallowWaterCollisionTracker_Actor(float InTimeSpawned, float InLifespan, TWeakObjectPtr<AActor> InCollisionActor)
 		: TimeSpawned(InTimeSpawned), Lifespan(InLifespan), CollisionActor(InCollisionActor)
 	{}
 
-	TArray<AWaterBody*> GetOverlappingWaterBodies() const;
+	void GetOverlappingWaterBodies(TSet<AWaterBody*>& WaterBodies) const;
 	bool IsValid(const float CurrentTime) const { return CollisionActor.IsValid() && CurrentTime - TimeSpawned <= Lifespan; }
 
 	float TimeSpawned = 0;
 	float Lifespan = 10.f;
-	UPROPERTY()
 	TWeakObjectPtr<AActor> CollisionActor = nullptr;
 
 	bool operator==(const FShallowWaterCollisionTracker_Actor& Other) const
@@ -97,11 +93,8 @@ struct WATERADVANCED_API FShallowWaterCollisionTracker_Actor
 /*
  * CollisionTracker that directly records the water body affected, used by Impacts
  */
-USTRUCT()
 struct WATERADVANCED_API FShallowWaterCollisionTracker_Direct
 {
-	GENERATED_BODY()
-
 	FShallowWaterCollisionTracker_Direct() {}
 	FShallowWaterCollisionTracker_Direct(float InTimeSpawned, float InLifespan, TWeakObjectPtr<AWaterBody> InWaterBody)
 		: TimeSpawned(InTimeSpawned), Lifespan(InLifespan), WaterBody(InWaterBody)
@@ -112,12 +105,10 @@ struct WATERADVANCED_API FShallowWaterCollisionTracker_Direct
 
 	float TimeSpawned = 0;
 	float Lifespan = 10.f;
-	UPROPERTY()
 	TWeakObjectPtr<AWaterBody> WaterBody = nullptr;
 };
 
 
-
 UCLASS(Abstract)
 class WATERADVANCED_API UShallowWaterSubsystem : public UTickableWorldSubsystem
 {
@@ -202,10 +193,11 @@ protected:
 	void RemoveCollisionTrackerForActor(AActor* CollisionTrackerActor);
 	
 	virtual TSet<AWaterBody*> GetOverlappingWaterBodiesFromPawns() const;
-	TSet<AWaterBody*> GetOverlappingWaterBodiesFromActorTrackersAndUpdate();
-	TSet<AWaterBody*> GetOverlappingWaterBodiesFromDirectTrackersAndUpdate();
+	void GetOverlappingWaterBodiesFromActorTrackersAndUpdate(TSet<AWaterBody*>& WaterBodies);
+	void GetOverlappingWaterBodiesFromDirectTrackersAndUpdate(TSet<AWaterBody*>& WaterBodies);
 	void UpdateOverlappingWaterBodiesHistory(TArray<AWaterBody*> OverlappingWaterBodies);
 	bool ShouldSimulateThisFrame() const;
+	void ClearTickCache();
 
 	/*
 	 * Collision
@@ -213,9 +205,8 @@ protected:
 
 	static FName ColliderComponentTag;
 	
-	// Override to get most relevant pawns every frame
+	// Override to get most relevant pawns every frame, this also sets up the per-tick cache that is reset with ClearTickCache
 	virtual int32 UpdateActivePawns();
-	virtual void GatherContextsFromPawns(TArray<TWeakObjectPtr<APawn>> ActivePawns);
 	// By default getting the first SKM, if is ACharacter, get 'Mesh' component
 	// Override if e.g. The pawn is driving a boat, return SKM of boat instead
 	virtual TOptional<FShallowWaterCollisionContext> GetCollisionContextFromPawn(APawn* InPawn) const;
@@ -251,12 +242,16 @@ protected:
 
 	// Consistent record through multiple frames since we want active pawns that affects the fluidsim to remain active
 	// Even when it's less significant than newly available pawns. Fluidsim enabled for different pawns each frame is bad.
-	UPROPERTY()
 	TArray<TWeakObjectPtr<APawn>> ActivePawns;
 
 	FVector PreviousProjectedLocation = FVector::ZeroVector;
 	TWeakObjectPtr<APlayerController> WeakPlayerController;
 
+	// Cached variables only valid during a tick
+	bool bTickCacheValid = false;
+	APawn* CachedCursorPawn = nullptr;
+	TArray<APawn*> CachedPawnsInRange;
+
 private:
 	float LastTimeOverlappingAnyWaterBody = -FLT_MAX;
 	// Should only be managed by UpdateOverlappingWaterBodiesHistory()
@@ -265,8 +260,8 @@ private:
 	bool bIsShallowWaterInitialized = false;
 
 	bool bInitializationAsyncLoadsAttempted = false;
-
-	TSet<const AWaterZone*> WaterZonesWithMaterialsInitialized;
+	
+	TSet<TWeakObjectPtr<AWaterZone>> WaterZonesWithMaterialsInitialized;
 
 	struct PendingImpact
 	{
@@ -278,10 +273,7 @@ private:
 	TArray<PendingImpact> PendingImpacts;
 	bool bFlushPendingImpactsNextTick;
 
-
-	UPROPERTY()
 	TSet<TWeakObjectPtr<UWaterBodyComponent>> WaterBodyComponentsWithProperMIDParameters;
-	UPROPERTY()
 	TArray<TWeakObjectPtr<AWaterBody>> PendingWaterBodiesToSetMIDOnInitialize;
 
 	/*
@@ -301,9 +293,7 @@ private:
 	UFUNCTION()
 	void OnLocalPlayerPawnBecomesValid(APawn* OldPawn, APawn* NewPawn);
 
-	UPROPERTY()
 	TArray<FShallowWaterCollisionTracker_Actor> Tracker_Actors;
-	UPROPERTY()
 	TArray<FShallowWaterCollisionTracker_Direct> Tracker_Directs;
 
 	// Overrides collected from RegisterPhysicsAssetOverridesDataAsset
