fix logspam
better error checking when missing baked sim

diff --git a/Engine/Plugins/Experimental/WaterAdvanced/Source/WaterAdvanced/Private/ShallowWaterRiverActor.cpp b/Engine/Plugins/Experimental/WaterAdvanced/Source/WaterAdvanced/Private/ShallowWaterRiverActor.cpp
--- a/Engine/Plugins/Experimental/WaterAdvanced/Source/WaterAdvanced/Private/ShallowWaterRiverActor.cpp
+++ b/Engine/Plugins/Experimental/WaterAdvanced/Source/WaterAdvanced/Private/ShallowWaterRiverActor.cpp
@@ -263,13 +263,13 @@ void UShallowWaterRiverComponent::Rebuild()
 	}
 	else	
 	{
-		UE_LOG(LogShallowWater, Warning, TEXT("UShallowWaterRiverComponent::Rebuild() - No source water bodies specified"));
+		UE_LOG(LogShallowWater, Verbose, TEXT("UShallowWaterRiverComponent::Rebuild() - No source water bodies specified"));
 		return;
 	}
 	
 	if (AllWaterBodies.Num() == 0)
 	{
-		UE_LOG(LogShallowWater, Warning, TEXT("UShallowWaterRiverComponent::Rebuild() - No valid source water bodies specified"));
+		UE_LOG(LogShallowWater, Verbose, TEXT("UShallowWaterRiverComponent::Rebuild() - No valid source water bodies specified"));
 		return;
 	}
 
@@ -1042,6 +1042,12 @@ void UShallowWaterRiverComponent::UpdateRenderState()
 		RiverSimSystem->ReinitializeSystem();
 	}
 
+	if (RenderState == EShallowWaterRenderState::BakedSim || RenderState == EShallowWaterRenderState::WaterComponentWithBakedSim && 
+		(BakedWaterSurfaceTexture == nullptr || BakedWaterSurfaceTexture->GetSizeX() == 0 || BakedWaterSurfaceTexture->GetSizeY() == 0))
+	{
+		UE_LOG(LogShallowWater, Warning, TEXT("UShallowWaterRiverComponent::UpdateRenderState() - No baked sim to render"));		
+	}
+
 	for (TSoftObjectPtr<AWaterBody > CurrWaterBody : AllWaterBodies)
 	{
 		if (!CurrWaterBody)
