fix cis

diff --git a/Engine/Plugins/Experimental/WaterAdvanced/Source/WaterAdvanced/Private/ShallowWaterRiverActor.cpp b/Engine/Plugins/Experimental/WaterAdvanced/Source/WaterAdvanced/Private/ShallowWaterRiverActor.cpp
--- a/Engine/Plugins/Experimental/WaterAdvanced/Source/WaterAdvanced/Private/ShallowWaterRiverActor.cpp
+++ b/Engine/Plugins/Experimental/WaterAdvanced/Source/WaterAdvanced/Private/ShallowWaterRiverActor.cpp
@@ -1042,7 +1042,7 @@ void UShallowWaterRiverComponent::UpdateRenderState()
 		RiverSimSystem->ReinitializeSystem();
 	}
 
-	if (RenderState == EShallowWaterRenderState::BakedSim || RenderState == EShallowWaterRenderState::WaterComponentWithBakedSim && 
+	if ((RenderState == EShallowWaterRenderState::BakedSim || RenderState == EShallowWaterRenderState::WaterComponentWithBakedSim) && 
 		(BakedWaterSurfaceTexture == nullptr || BakedWaterSurfaceTexture->GetSizeX() == 0 || BakedWaterSurfaceTexture->GetSizeY() == 0))
 	{
 		UE_LOG(LogShallowWater, Warning, TEXT("UShallowWaterRiverComponent::UpdateRenderState() - No baked sim to render"));		
