better initial settings for sw rivers

diff --git a/Engine/Plugins/Experimental/WaterAdvanced/Source/WaterAdvanced/Private/ShallowWaterRiverActor.cpp b/Engine/Plugins/Experimental/WaterAdvanced/Source/WaterAdvanced/Private/ShallowWaterRiverActor.cpp
--- a/Engine/Plugins/Experimental/WaterAdvanced/Source/WaterAdvanced/Private/ShallowWaterRiverActor.cpp
+++ b/Engine/Plugins/Experimental/WaterAdvanced/Source/WaterAdvanced/Private/ShallowWaterRiverActor.cpp
@@ -52,21 +52,6 @@ UShallowWaterRiverComponent::UShallowWaterRiverComponent(const FObjectInitialize
 
 	ResolutionMaxAxis = 512;
 	SourceSize = 1000;
-
-#if WITH_EDITOR	
-	// start with one empty element
-	if (BottomContourActors.IsEmpty())
-	{
-		BottomContourActors.AddDefaulted();
-	}
-
-	// start with one empty element
-	if (DilatedBottomContourActors.IsEmpty())
-	{
-		DilatedBottomContourActors.AddDefaulted();
-	}
-#endif
-
 }
 
 TObjectPtr<UTextureRenderTarget2D> UShallowWaterRiverComponent::GetSharedFFTOceanPatchNormalRTFromSubsystem(UWorld* World)
diff --git a/Engine/Plugins/Experimental/WaterAdvanced/Source/WaterAdvanced/Public/ShallowWaterRiverActor.h b/Engine/Plugins/Experimental/WaterAdvanced/Source/WaterAdvanced/Public/ShallowWaterRiverActor.h
--- a/Engine/Plugins/Experimental/WaterAdvanced/Source/WaterAdvanced/Public/ShallowWaterRiverActor.h
+++ b/Engine/Plugins/Experimental/WaterAdvanced/Source/WaterAdvanced/Public/ShallowWaterRiverActor.h
@@ -59,10 +59,10 @@ public:
 	bool bMatchSpline = true;
 
 	UPROPERTY(EditAnywhere, Category = "Simulation", meta = (DisplayName = "Match Spline Height Amount", EditCondition = "bMatchSpline"))
-	float MatchSplineHeightAmount = 5.;
+	float MatchSplineHeightAmount = 2.;
 
 	UPROPERTY(EditAnywhere, Category = "Simulation", meta = (DisplayName = "Remove Outside Spline Amount", EditCondition = "bMatchSpline"))
-	float RemoveOutsideSplineAmount = 1000.;
+	float RemoveOutsideSplineAmount = 50.;
 
 	UPROPERTY(EditAnywhere, Category = "Water", meta = (DisplayName = "Source River Water Bodies"))
 	TArray<TSoftObjectPtr<AWaterBody>> SourceRiverWaterBodies;
