expose landscape colliders separately
default landscape colliders array populated from map

diff --git a/Engine/Plugins/Experimental/WaterAdvanced/Source/WaterAdvanced/Private/ShallowWaterRiverActor.cpp b/Engine/Plugins/Experimental/WaterAdvanced/Source/WaterAdvanced/Private/ShallowWaterRiverActor.cpp
--- a/Engine/Plugins/Experimental/WaterAdvanced/Source/WaterAdvanced/Private/ShallowWaterRiverActor.cpp
+++ b/Engine/Plugins/Experimental/WaterAdvanced/Source/WaterAdvanced/Private/ShallowWaterRiverActor.cpp
@@ -52,6 +52,15 @@ UShallowWaterRiverComponent::UShallowWaterRiverComponent(const FObjectInitialize
 
 	ResolutionMaxAxis = 512;
 	SourceSize = 1000;
+	
+	// initialize landscape array with all landscapes
+	if (GetWorld())
+	{
+		for (TActorIterator<ALandscape> It(GetWorld(), ALandscape::StaticClass()); It; ++It)
+		{
+			BottomContourLandscapeActors.Add(*It);
+		}
+	}
 }
 
 TObjectPtr<UTextureRenderTarget2D> UShallowWaterRiverComponent::GetSharedFFTOceanPatchNormalRTFromSubsystem(UWorld* World)
@@ -331,6 +340,11 @@ void UShallowWaterRiverComponent::Rebuild()
 			TArray<AActor*> BottomContourActorsRawPtr;
 			BottomContourActorsRawPtr.Add(nullptr);
 
+			for (TSoftObjectPtr<ALandscape> CurrLandscapeActor : BottomContourLandscapeActors)
+			{
+				BottomContourActorsRawPtr.Add(CurrLandscapeActor.Get());
+			}
+		
 			AddActorsToRawArray(BottomContourActors, BottomContourActorsRawPtr);			
 			AddTaggedActorsToArray(BottomContourTags, BottomContourActorsRawPtr);
 			FName DIName = "User.BottomCapture";
@@ -345,7 +359,6 @@ void UShallowWaterRiverComponent::Rebuild()
 				false,
 				BottomContourActorsRawPtr);
 
-	
 			// accumulate bounding box for river water bodies
 			FBoxSphereBounds::Builder BottomContourCombinedWorldBoundsBuilder;
 			for (TSoftObjectPtr<AActor> BottomContourActor : BottomContourActors)
@@ -360,7 +373,7 @@ void UShallowWaterRiverComponent::Rebuild()
 				}
 				else
 				{
-					UE_LOG(LogShallowWater, Warning, TEXT("UShallowWaterRiverComponent::Rebuild() - skipping null bottom contour boundary actor found"));
+					UE_LOG(LogShallowWater, Verbose, TEXT("UShallowWaterRiverComponent::Rebuild() - skipping null bottom contour boundary actor found"));
 					continue;
 				}
 			}
@@ -384,7 +397,6 @@ void UShallowWaterRiverComponent::Rebuild()
 				false,
 				DilatedBottomContourActorsRawPtr);
 
-
 			// accumulate bounding box for river water bodies
 			FBoxSphereBounds::Builder DilatedBottomContourCombinedWorldBoundsBuilder;
 			for (AActor *BottomContourActor : DilatedBottomContourActorsRawPtr)
@@ -399,7 +411,7 @@ void UShallowWaterRiverComponent::Rebuild()
 				}
 				else
 				{
-					UE_LOG(LogShallowWater, Warning, TEXT("UShallowWaterRiverComponent::Rebuild() - skipping null bottom contour boundary actor found"));
+					UE_LOG(LogShallowWater, Verbose, TEXT("UShallowWaterRiverComponent::Rebuild() - skipping null bottom contour boundary actor found"));
 					continue;
 				}
 			}
diff --git a/Engine/Plugins/Experimental/WaterAdvanced/Source/WaterAdvanced/Public/ShallowWaterRiverActor.h b/Engine/Plugins/Experimental/WaterAdvanced/Source/WaterAdvanced/Public/ShallowWaterRiverActor.h
--- a/Engine/Plugins/Experimental/WaterAdvanced/Source/WaterAdvanced/Public/ShallowWaterRiverActor.h
+++ b/Engine/Plugins/Experimental/WaterAdvanced/Source/WaterAdvanced/Public/ShallowWaterRiverActor.h
@@ -20,6 +20,7 @@ class UNiagaraComponent;
 class UNiagaraSystem;
 class AWaterBody;
 class AWaterZone;
+class ALandscape;
 
 UENUM(BlueprintType)
 enum EShallowWaterRenderState : int
@@ -109,6 +110,9 @@ public:
 	UPROPERTY(EditAnywhere, Category = "Collisions")
 	bool bUseCapture = true;
 
+	UPROPERTY(EditAnywhere, Category = "Collisions")
+	TArray<TSoftObjectPtr<ALandscape>> BottomContourLandscapeActors;
+
 	UPROPERTY(EditAnywhere, Category = "Collisions")
 	TArray<TSoftObjectPtr<AActor>> BottomContourActors;
 
