better default for bottom contour capture offset

diff --git a/Engine/Plugins/Experimental/WaterAdvanced/Source/WaterAdvanced/Public/ShallowWaterRiverActor.h b/Engine/Plugins/Experimental/WaterAdvanced/Source/WaterAdvanced/Public/ShallowWaterRiverActor.h
--- a/Engine/Plugins/Experimental/WaterAdvanced/Source/WaterAdvanced/Public/ShallowWaterRiverActor.h
+++ b/Engine/Plugins/Experimental/WaterAdvanced/Source/WaterAdvanced/Public/ShallowWaterRiverActor.h
@@ -120,7 +120,7 @@ public:
 	TArray<FName> BottomContourTags;
 
 	UPROPERTY(EditAnywhere, Category = "Collisions")
-	float BottomContourCaptureOffset = 5000.f;
+	float BottomContourCaptureOffset = 15000.f;
 
 	UPROPERTY(EditAnywhere, Category = "Collisions")
 	TArray<TSoftObjectPtr<AActor>> DilatedBottomContourActors;
