expose pause on sw river component

diff --git a/Engine/Plugins/Experimental/WaterAdvanced/Source/WaterAdvanced/Private/ShallowWaterRiverActor.cpp b/Engine/Plugins/Experimental/WaterAdvanced/Source/WaterAdvanced/Private/ShallowWaterRiverActor.cpp
--- a/Engine/Plugins/Experimental/WaterAdvanced/Source/WaterAdvanced/Private/ShallowWaterRiverActor.cpp
+++ b/Engine/Plugins/Experimental/WaterAdvanced/Source/WaterAdvanced/Private/ShallowWaterRiverActor.cpp
@@ -976,9 +976,26 @@ void UShallowWaterRiverComponent::OnWaterInfoTextureArrayCreated(const UTextureR
 		return;
 	}
 }
-
 #endif
 
+void UShallowWaterRiverComponent::SetPaused(bool Pause)
+{
+	if (RiverSimSystem)
+	{
+		RiverSimSystem->SetPaused(Pause);
+	}
+
+	UFFTOceanPatchSubsystem *OceanPatchSubsystem = GetWorld()->GetSubsystem<UFFTOceanPatchSubsystem>();
+	if (OceanPatchSubsystem != nullptr)
+	{
+		TObjectPtr<UNiagaraComponent> OceanSystem = OceanPatchSubsystem->GetOceanSystem();
+		if (OceanSystem)
+		{
+			OceanSystem->SetPaused(Pause);
+		}
+	}
+}
+
 void UShallowWaterRiverComponent::UpdateRenderState()
 {
 	TObjectPtr<UTextureRenderTarget2D>  OceanPatchNormalRT = GetSharedFFTOceanPatchNormalRTFromSubsystem(GetWorld());
diff --git a/Engine/Plugins/Experimental/WaterAdvanced/Source/WaterAdvanced/Public/FFTOceanPatchSubsystem.h b/Engine/Plugins/Experimental/WaterAdvanced/Source/WaterAdvanced/Public/FFTOceanPatchSubsystem.h
--- a/Engine/Plugins/Experimental/WaterAdvanced/Source/WaterAdvanced/Public/FFTOceanPatchSubsystem.h
+++ b/Engine/Plugins/Experimental/WaterAdvanced/Source/WaterAdvanced/Public/FFTOceanPatchSubsystem.h
@@ -24,6 +24,7 @@ public:
 	UE_API virtual TStatId GetStatId() const override;
 	
 	UE_API TObjectPtr<UTextureRenderTarget2D> GetOceanNormalRT(UWorld* World);
+	UE_API TObjectPtr<UNiagaraComponent> GetOceanSystem() { return FFTOceanSystem; }
 protected:
 	UPROPERTY(BlueprintReadOnly, VisibleAnywhere, Category="Shallow Water")
 	TObjectPtr<UShallowWaterSettings> Settings;
diff --git a/Engine/Plugins/Experimental/WaterAdvanced/Source/WaterAdvanced/Public/ShallowWaterRiverActor.h b/Engine/Plugins/Experimental/WaterAdvanced/Source/WaterAdvanced/Public/ShallowWaterRiverActor.h
--- a/Engine/Plugins/Experimental/WaterAdvanced/Source/WaterAdvanced/Public/ShallowWaterRiverActor.h
+++ b/Engine/Plugins/Experimental/WaterAdvanced/Source/WaterAdvanced/Public/ShallowWaterRiverActor.h
@@ -157,6 +157,9 @@ public:
 	UE_API void OnWaterInfoTextureArrayCreated(const UTextureRenderTarget2DArray* InWaterInfoTexture);
 #endif // WITH_EDITOR
 
+	UFUNCTION(BlueprintCallable, Category = "Shallow Water River")
+	UE_API void SetPaused(bool Pause);
+
 protected:
 	// Asset can be set in Project Settings - Plugins - Water ShallowWaterSimulation
 	UPROPERTY(BlueprintReadOnly, VisibleDefaultsOnly, Category = "Shallow Water")
