refactoring for capture modules in sw emitters
handle overhangs in terrain with better controls
rework summary views

#virtualized

diff --git a/Engine/Plugins/Experimental/WaterAdvanced/Source/WaterAdvanced/Private/ShallowWaterRiverActor.cpp b/Engine/Plugins/Experimental/WaterAdvanced/Source/WaterAdvanced/Private/ShallowWaterRiverActor.cpp
--- a/Engine/Plugins/Experimental/WaterAdvanced/Source/WaterAdvanced/Private/ShallowWaterRiverActor.cpp
+++ b/Engine/Plugins/Experimental/WaterAdvanced/Source/WaterAdvanced/Private/ShallowWaterRiverActor.cpp
@@ -86,6 +86,40 @@ TObjectPtr<UTextureRenderTarget2D> UShallowWaterRiverComponent::GetSharedFFTOcea
 	return nullptr;
 }
 
+FBoxSphereBounds UShallowWaterRiverComponent::InitializeCaptureDI(const FName& DIName, TArray<AActor*> RawActorPtrArray)
+{
+	UNiagaraFunctionLibrary::SetSceneCapture2DDataInterfaceManagedMode(RiverSimSystem, DIName,
+				ESceneCaptureSource::SCS_SceneDepth,
+				FIntPoint(ResolutionMaxAxis, ResolutionMaxAxis),
+				ETextureRenderTargetFormat::RTF_R32f,
+				ECameraProjectionMode::Orthographic,
+				90.0f,
+				FMath::Max(WorldGridSize.X, WorldGridSize.Y),
+				true,
+				false,
+				RawActorPtrArray);
+
+	// accumulate bounding box for river water bodies
+	FBoxSphereBounds::Builder BottomContourCombinedWorldBoundsBuilder;
+	for (AActor *BottomContourActor : RawActorPtrArray)
+	{
+		if (BottomContourActor != nullptr)
+		{
+			// accumulate bounds
+			FBoxSphereBounds WorldBounds;
+			BottomContourActor->GetActorBounds(false, WorldBounds.Origin, WorldBounds.BoxExtent);
+
+			BottomContourCombinedWorldBoundsBuilder += WorldBounds;
+		}
+		else
+		{
+			UE_LOG(LogShallowWater, Verbose, TEXT("UShallowWaterRiverComponent::Rebuild() - skipping null bottom contour boundary actor found"));
+			continue;
+		}
+	}
+	return FBoxSphereBounds(BottomContourCombinedWorldBoundsBuilder);	
+}
+
 void UShallowWaterRiverComponent::PostLoad()
 {
 	Super::PostLoad();	
@@ -337,90 +371,42 @@ void UShallowWaterRiverComponent::Rebuild()
 		// convert to raw ptr array for function library
 		if (!bReadBakedSim && bUseCapture)
 		{
-			TArray<AActor*> BottomContourActorsRawPtr;
-			BottomContourActorsRawPtr.Add(nullptr);
-
+			// landscape captures
+			TArray<AActor*> LandscapeBottomContourActorsRawPtr;
+			LandscapeBottomContourActorsRawPtr.Add(nullptr);
 			for (TSoftObjectPtr<ALandscape> CurrLandscapeActor : BottomContourLandscapeActors)
 			{
-				BottomContourActorsRawPtr.Add(CurrLandscapeActor.Get());
+				LandscapeBottomContourActorsRawPtr.Add(CurrLandscapeActor.Get());
 			}
+			FBoxSphereBounds LandscapeBottomContourBounds = InitializeCaptureDI("User.LandscapeBottomCapture", LandscapeBottomContourActorsRawPtr);
 		
+			// undilated captures
+			TArray<AActor*> BottomContourActorsRawPtr;
+			BottomContourActorsRawPtr.Add(nullptr);
 			AddActorsToRawArray(BottomContourActors, BottomContourActorsRawPtr);			
 			AddTaggedActorsToArray(BottomContourTags, BottomContourActorsRawPtr);
-			FName DIName = "User.BottomCapture";
-			UNiagaraFunctionLibrary::SetSceneCapture2DDataInterfaceManagedMode(RiverSimSystem, DIName,
-				ESceneCaptureSource::SCS_SceneDepth,
-				FIntPoint(ResolutionMaxAxis, ResolutionMaxAxis),
-				ETextureRenderTargetFormat::RTF_R32f,
-				ECameraProjectionMode::Orthographic,
-				90.0f,
-				FMath::Max(WorldGridSize.X, WorldGridSize.Y),
-				true,
-				false,
-				BottomContourActorsRawPtr);
+			FBoxSphereBounds CombinedBottomContourBounds = InitializeCaptureDI("User.BottomCapture", BottomContourActorsRawPtr);
+			FBoxSphereBounds CombinedBottomContourBoundsUnder = InitializeCaptureDI("User.BottomCaptureUnder", BottomContourActorsRawPtr);
 
-			// accumulate bounding box for river water bodies
-			FBoxSphereBounds::Builder BottomContourCombinedWorldBoundsBuilder;
-			for (TSoftObjectPtr<AActor> BottomContourActor : BottomContourActors)
-			{
-				if (BottomContourActor != nullptr)
-				{
-					// accumulate bounds
-					FBoxSphereBounds WorldBounds;
-					BottomContourActor->GetActorBounds(false, WorldBounds.Origin, WorldBounds.BoxExtent);
-
-					BottomContourCombinedWorldBoundsBuilder += WorldBounds;
-				}
-				else
-				{
-					UE_LOG(LogShallowWater, Verbose, TEXT("UShallowWaterRiverComponent::Rebuild() - skipping null bottom contour boundary actor found"));
-					continue;
-				}
-			}
-			FBoxSphereBounds CombinedBottomContourBounds(BottomContourCombinedWorldBoundsBuilder);				
-			
 			// Dilated capture
 			TArray<AActor*> DilatedBottomContourActorsRawPtr;
 			DilatedBottomContourActorsRawPtr.Add(nullptr);
 			AddActorsToRawArray(DilatedBottomContourActors, DilatedBottomContourActorsRawPtr);
 			AddTaggedActorsToArray(DilatedBottomContourTags, DilatedBottomContourActorsRawPtr);
 
-			FName DilatedDIName = "User.DilatedBottomCapture";
-			UNiagaraFunctionLibrary::SetSceneCapture2DDataInterfaceManagedMode(RiverSimSystem, DilatedDIName,
-				ESceneCaptureSource::SCS_SceneDepth,
-				FIntPoint(ResolutionMaxAxis, ResolutionMaxAxis),
-				ETextureRenderTargetFormat::RTF_R32f,
-				ECameraProjectionMode::Orthographic,
-				90.0f,
-				FMath::Max(WorldGridSize.X, WorldGridSize.Y),
-				true,
-				false,
-				DilatedBottomContourActorsRawPtr);
+			FBoxSphereBounds DilatedCombinedBottomContourBounds = InitializeCaptureDI("User.DilatedBottomCapture", DilatedBottomContourActorsRawPtr);
+			FBoxSphereBounds DilatedCombinedBottomContourBoundsUnder = InitializeCaptureDI("User.DilatedBottomCaptureUnder", DilatedBottomContourActorsRawPtr);
 
-			// accumulate bounding box for river water bodies
-			FBoxSphereBounds::Builder DilatedBottomContourCombinedWorldBoundsBuilder;
-			for (AActor *BottomContourActor : DilatedBottomContourActorsRawPtr)
-			{
-				if (BottomContourActor != nullptr)
-				{
-					// accumulate bounds
-					FBoxSphereBounds WorldBounds;
-					BottomContourActor->GetActorBounds(false, WorldBounds.Origin, WorldBounds.BoxExtent);
-
-					DilatedBottomContourCombinedWorldBoundsBuilder += WorldBounds;
-				}
-				else
-				{
-					UE_LOG(LogShallowWater, Verbose, TEXT("UShallowWaterRiverComponent::Rebuild() - skipping null bottom contour boundary actor found"));
-					continue;
-				}
-			}
-			FBoxSphereBounds DilatedCombinedBottomContourBounds(DilatedBottomContourCombinedWorldBoundsBuilder);
-			
 			// reinitialize and set variables on the system
 			RiverSimSystem->ReinitializeSystem();
-			RiverSimSystem->SetVariableFloat(FName("CaptureOffset"), BottomContourCaptureOffset + CombinedBottomContourBounds.Origin.Z + CombinedBottomContourBounds.BoxExtent.Z);
-			RiverSimSystem->SetVariableFloat(FName("DilatedCaptureOffset"), BottomContourCaptureOffset + DilatedCombinedBottomContourBounds.Origin.Z + DilatedCombinedBottomContourBounds.BoxExtent.Z);
+
+			RiverSimSystem->SetVariableFloat(FName("LandscapeCaptureOffset"), LandscapeBottomContourBounds.Origin.Z + LandscapeBottomContourBounds.BoxExtent.Z + BottomContourCaptureOffset);			
+
+			RiverSimSystem->SetVariableFloat(FName("CaptureOffset"), CombinedBottomContourBounds.Origin.Z + CombinedBottomContourBounds.BoxExtent.Z + BottomContourCaptureOffset);
+			RiverSimSystem->SetVariableFloat(FName("DilatedCaptureOffset"), DilatedCombinedBottomContourBounds.Origin.Z + DilatedCombinedBottomContourBounds.BoxExtent.Z + BottomContourCaptureOffset);
+			
+			RiverSimSystem->SetVariableFloat(FName("CaptureOffsetUnder"), CombinedBottomContourBounds.Origin.Z - CombinedBottomContourBounds.BoxExtent.Z - BottomContourCaptureOffset);
+			RiverSimSystem->SetVariableFloat(FName("DilatedCaptureOffsetUnder"), DilatedCombinedBottomContourBounds.Origin.Z - DilatedCombinedBottomContourBounds.BoxExtent.Z - BottomContourCaptureOffset);
 		}
 		else
 		{
diff --git a/Engine/Plugins/Experimental/WaterAdvanced/Source/WaterAdvanced/Public/ShallowWaterRiverActor.h b/Engine/Plugins/Experimental/WaterAdvanced/Source/WaterAdvanced/Public/ShallowWaterRiverActor.h
--- a/Engine/Plugins/Experimental/WaterAdvanced/Source/WaterAdvanced/Public/ShallowWaterRiverActor.h
+++ b/Engine/Plugins/Experimental/WaterAdvanced/Source/WaterAdvanced/Public/ShallowWaterRiverActor.h
@@ -204,6 +204,8 @@ private:
 
 	// Find/create the level set renderer singleton actor as required. Return whether the found or created actor.
 	TObjectPtr<UTextureRenderTarget2D> GetSharedFFTOceanPatchNormalRTFromSubsystem(UWorld* World);
+
+	FBoxSphereBounds InitializeCaptureDI(const FName &DIName, TArray<AActor*> RawActorPtrArray);
 };
 
 UCLASS(MinimalAPI, BlueprintType, HideCategories = (Physics, Replication, Input, Collision))
