fix extrapolate render height
update templates to latest emitter versions

#virtualized

diff --git a/Engine/Plugins/Experimental/WaterAdvanced/Source/WaterAdvanced/Public/ShallowWaterRiverActor.h b/Engine/Plugins/Experimental/WaterAdvanced/Source/WaterAdvanced/Public/ShallowWaterRiverActor.h
--- a/Engine/Plugins/Experimental/WaterAdvanced/Source/WaterAdvanced/Public/ShallowWaterRiverActor.h
+++ b/Engine/Plugins/Experimental/WaterAdvanced/Source/WaterAdvanced/Public/ShallowWaterRiverActor.h
@@ -93,7 +93,7 @@ public:
 	TObjectPtr <class UMaterialInstance> SplineRiverToOceanTransitionMaterial;
 
 	UPROPERTY(EditAnywhere, Category = "Rendering", meta = (DisplayName = "Smoothing Width"))
-	int SmoothingWidth = 2;
+	int SmoothingWidth = 5;
 
 	UPROPERTY(EditAnywhere, Category = "Rendering", meta = (DisplayName = "Smoothing Cutoff"))
 	float SmoothingCutoff = 500.f;
