Unshelved from pending changelist '42455702':

cache normals, rework sim to accomodate
avoids bug with texture streaming and reset stages on Niagara Emitters

diff --git a/Engine/Plugins/Experimental/WaterAdvanced/Source/WaterAdvanced/Private/ShallowWaterRiverActor.cpp b/Engine/Plugins/Experimental/WaterAdvanced/Source/WaterAdvanced/Private/ShallowWaterRiverActor.cpp
--- a/Engine/Plugins/Experimental/WaterAdvanced/Source/WaterAdvanced/Private/ShallowWaterRiverActor.cpp
+++ b/Engine/Plugins/Experimental/WaterAdvanced/Source/WaterAdvanced/Private/ShallowWaterRiverActor.cpp
@@ -585,6 +585,10 @@ void UShallowWaterRiverComponent::Rebuild()
 	BakedFoamRT->InitAutoFormat(1, 1);
 	RiverSimSystem->SetVariableTextureRenderTarget(FName("FoamRT"), BakedFoamRT);
 	
+	BakedWaterSurfaceNormalRT = NewObject<UTextureRenderTarget2D>(this, NAME_None, RF_Transient);
+	BakedWaterSurfaceNormalRT->InitAutoFormat(1, 1);
+	RiverSimSystem->SetVariableTextureRenderTarget(FName("NormalRT"), BakedWaterSurfaceNormalRT);
+
 	RiverSimSystem->SetVariableBool(FName("ReadCachedSim"), bReadBakedSim);
 
 	RiverSimSystem->SetVariableFloat(FName("BottomContourCollisionDilation"), BottomContourCollisionDilation);
@@ -592,16 +596,11 @@ void UShallowWaterRiverComponent::Rebuild()
 	RiverSimSystem->SetVariableInt(FName("ExtrapolationHalfWidth"), SmoothingWidth);
 	RiverSimSystem->SetVariableFloat(FName("SmoothingHeightCutoff"), SmoothingCutoff);
 
-	if (BakedWaterSurfaceTexture != nullptr && BakedFoamTexture != nullptr)
+	if (BakedWaterSurfaceTexture != nullptr && BakedFoamTexture != nullptr && BakedWaterSurfaceNormalTexture != nullptr)
 	{
 		RiverSimSystem->SetVariableTexture(FName("BakedSimTexture"), BakedWaterSurfaceTexture);
 		RiverSimSystem->SetVariableTexture(FName("BakedFoamTexture"), BakedFoamTexture);
-	
-		FIntVector2 NumCells = FIntVector2(BakedWaterSurfaceTexture->GetSizeX(), BakedWaterSurfaceTexture->GetSizeY());
-
-		BakedWaterSurfaceNormalRT = NewObject<UTextureRenderTarget2D>(this);	
-		BakedWaterSurfaceNormalRT->InitAutoFormat(1, 1);
-		RiverSimSystem->SetVariableTextureRenderTarget(FName("NormalRT"), BakedWaterSurfaceNormalRT);
+		RiverSimSystem->SetVariableTexture(FName("BakedWaterSurfaceNormalTexture"), BakedWaterSurfaceNormalTexture);
 	}
 
 	TObjectPtr<UTextureRenderTarget2D>  OceanPatchNormalRT = GetSharedFFTOceanPatchNormalRTFromSubsystem(GetWorld());
@@ -691,7 +690,7 @@ void UShallowWaterRiverComponent::Bake()
 {
 	EObjectFlags TextureObjectFlags = EObjectFlags::RF_Public;
 
-	if (!RiverSimSystem || !BakedWaterSurfaceRT || !BakedFoamRT)
+	if (!RiverSimSystem || !BakedWaterSurfaceRT || !BakedFoamRT || !BakedWaterSurfaceNormalRT)
 	{
 		UE_LOG(LogShallowWater, Warning, TEXT("UShallowWaterRiverComponent::Bake() - No simulation to bake"));
 		return;
@@ -742,6 +741,10 @@ void UShallowWaterRiverComponent::Bake()
 	BakedFoamTexture = BakedFoamRT->ConstructTexture2D(this, "BakedFoamTexture", TextureObjectFlags);
 	RiverSimSystem->SetVariableTexture(FName("BakedFoamTexture"), BakedFoamTexture);
 	
+	// bake normal to texture
+	BakedWaterSurfaceNormalTexture = BakedWaterSurfaceNormalRT->ConstructTexture2D(this, "BakedWaterSurfaceNormalTexture", TextureObjectFlags);
+	RiverSimSystem->SetVariableTexture(FName("BakedWaterSurfaceNormalTexture"), BakedWaterSurfaceNormalTexture);
+
 	// clear references to old baked sim on water body actors
 	if (BakedSim != nullptr)
 	{ 
@@ -1081,6 +1084,7 @@ void UShallowWaterRiverComponent::UpdateRenderState()
 				{
 					WaterInfoMID->SetTextureParameterValue("BakedWaterSimTex", BakedWaterSurfaceTexture);
 					WaterInfoMID->SetTextureParameterValue("FoamTex", BakedFoamTexture);
+					WaterInfoMID->SetTextureParameterValue("BakedWaterSimNormalTex", BakedWaterSurfaceNormalTexture);
 					WaterInfoMID->SetVectorParameterValue("BakedWaterSimLocation", SystemPos);
 					WaterMID->SetDoubleVectorParameterValue("BakedWaterSimLocationDouble", SystemPos);
 					WaterInfoMID->SetVectorParameterValue("BakedWaterSimSize", FVector(WorldGridSize.X, WorldGridSize.Y, 1));
@@ -1127,12 +1131,12 @@ void UShallowWaterRiverComponent::SetWaterMIDParameters(UMaterialInstanceDynamic
 	{
 		WaterMID->SetTextureParameterValue("BakedWaterSimTex", BakedWaterSurfaceTexture);
 		WaterMID->SetTextureParameterValue("FoamTex", BakedFoamTexture);
+		WaterMID->SetTextureParameterValue("BakedWaterSimNormalTex", BakedWaterSurfaceNormalTexture);
 
 		WaterMID->SetVectorParameterValue("BakedWaterSimLocation", SystemPos);
 		WaterMID->SetDoubleVectorParameterValue("BakedWaterSimLocationDouble", SystemPos);
 		WaterMID->SetVectorParameterValue("BakedWaterSimSize", FVector(WorldGridSize.X, WorldGridSize.Y, 1));
 
-		WaterMID->SetTextureParameterValue("BakedWaterSimNormalTex", BakedWaterSurfaceNormalRT);
 		WaterMID->SetTextureParameterValue("NormalDetailTex", NormalDetailRT);
 
 		float Dx;
diff --git a/Engine/Plugins/Experimental/WaterAdvanced/Source/WaterAdvanced/Public/ShallowWaterRiverActor.h b/Engine/Plugins/Experimental/WaterAdvanced/Source/WaterAdvanced/Public/ShallowWaterRiverActor.h
--- a/Engine/Plugins/Experimental/WaterAdvanced/Source/WaterAdvanced/Public/ShallowWaterRiverActor.h
+++ b/Engine/Plugins/Experimental/WaterAdvanced/Source/WaterAdvanced/Public/ShallowWaterRiverActor.h
@@ -104,6 +104,9 @@ public:
 	UPROPERTY(EditAnywhere, Category = "Shallow Water")
 	TObjectPtr<UTexture2D> BakedFoamTexture;
 
+	UPROPERTY(EditAnywhere, Category = "Shallow Water")
+	TObjectPtr<UTexture2D> BakedWaterSurfaceNormalTexture;
+
 	//UPROPERTY(EditAnywhere, Category = "Shallow Water")
 	//TObjectPtr<UTexture2D> SignedDistanceToSplineTexture;
 
@@ -174,6 +177,9 @@ protected:
 	UPROPERTY(BlueprintReadOnly, VisibleDefaultsOnly, Category = "Shallow Water")
 	TObjectPtr<UTextureRenderTarget2D> BakedFoamRT;
 
+	UPROPERTY(BlueprintReadOnly, VisibleDefaultsOnly, Category = "Shallow Water")
+	TObjectPtr<UTextureRenderTarget2D> BakedWaterSurfaceNormalRT;
+
 	UPROPERTY(BlueprintReadOnly, VisibleDefaultsOnly, Category = "Shallow Water")
 	TObjectPtr<UBakedShallowWaterSimulationComponent> BakedSim;
 
@@ -193,12 +199,6 @@ private:
 	UPROPERTY()
 	FVector SystemPos;
 	
-	UPROPERTY()
-	TObjectPtr<UTextureRenderTarget2D> BakedWaterSurfaceNormalRT;
-
-	UPROPERTY()
-	TObjectPtr<UTextureRenderTarget2D> FoamRT;
-
 	UPROPERTY()
 	TObjectPtr<UTextureRenderTarget2D> NormalDetailRT;
 
