support landscape proxies

diff --git a/Engine/Plugins/Experimental/WaterAdvanced/Source/WaterAdvanced/Private/ShallowWaterRiverActor.cpp b/Engine/Plugins/Experimental/WaterAdvanced/Source/WaterAdvanced/Private/ShallowWaterRiverActor.cpp
--- a/Engine/Plugins/Experimental/WaterAdvanced/Source/WaterAdvanced/Private/ShallowWaterRiverActor.cpp
+++ b/Engine/Plugins/Experimental/WaterAdvanced/Source/WaterAdvanced/Private/ShallowWaterRiverActor.cpp
@@ -29,6 +29,7 @@
 #include "Math/Float16Color.h"
 #include "Kismet/GameplayStatics.h"
 #include "Landscape.h"
+#include "LandscapeStreamingProxy.h"
 #include "LevelInstance/LevelInstanceActor.h"
 #include "LevelInstance/LevelInstanceSubsystem.h"
 
@@ -60,6 +61,11 @@ UShallowWaterRiverComponent::UShallowWaterRiverComponent(const FObjectInitialize
 		{
 			BottomContourLandscapeActors.Add(*It);
 		}
+
+		for (TActorIterator<ALandscapeStreamingProxy> It(GetWorld(), ALandscapeStreamingProxy::StaticClass()); It; ++It)
+		{
+			BottomContourLandscapeActors.Add(*It);
+		}
 	}
 }
 
@@ -374,8 +380,15 @@ void UShallowWaterRiverComponent::Rebuild()
 			// landscape captures
 			TArray<AActor*> LandscapeBottomContourActorsRawPtr;
 			LandscapeBottomContourActorsRawPtr.Add(nullptr);
-			for (TSoftObjectPtr<ALandscape> CurrLandscapeActor : BottomContourLandscapeActors)
+			for (TSoftObjectPtr<AActor> CurrLandscapeActor : BottomContourLandscapeActors)
 			{
+				// only accept Landscapes and LandscapeStreamingProxies
+				if (!Cast<ALandscape>(CurrLandscapeActor.Get()) && !Cast<ALandscapeStreamingProxy>(CurrLandscapeActor.Get()))
+				{
+					UE_LOG(LogShallowWater, Warning, TEXT("UShallowWaterRiverComponent::Rebuild() - Landscape bottom contour actors can only be ALandscape actors or ALandscapeStreamingProxy actors"));
+					continue;
+				}
+
 				LandscapeBottomContourActorsRawPtr.Add(CurrLandscapeActor.Get());
 			}
 			FBoxSphereBounds LandscapeBottomContourBounds = InitializeCaptureDI("User.LandscapeBottomCapture", LandscapeBottomContourActorsRawPtr);
diff --git a/Engine/Plugins/Experimental/WaterAdvanced/Source/WaterAdvanced/Public/ShallowWaterRiverActor.h b/Engine/Plugins/Experimental/WaterAdvanced/Source/WaterAdvanced/Public/ShallowWaterRiverActor.h
--- a/Engine/Plugins/Experimental/WaterAdvanced/Source/WaterAdvanced/Public/ShallowWaterRiverActor.h
+++ b/Engine/Plugins/Experimental/WaterAdvanced/Source/WaterAdvanced/Public/ShallowWaterRiverActor.h
@@ -114,7 +114,7 @@ public:
 	bool bUseCapture = true;
 
 	UPROPERTY(EditAnywhere, Category = "Collisions")
-	TArray<TSoftObjectPtr<ALandscape>> BottomContourLandscapeActors;
+	TArray<TSoftObjectPtr<AActor>> BottomContourLandscapeActors;
 
 	UPROPERTY(EditAnywhere, Category = "Collisions")
 	TArray<TSoftObjectPtr<AActor>> BottomContourActors;
