Add new icon
Set WaterAdvanced to beta

diff --git a/Engine/Plugins/Experimental/WaterAdvanced/WaterAdvanced.uplugin b/Engine/Plugins/Experimental/WaterAdvanced/WaterAdvanced.uplugin
--- a/Engine/Plugins/Experimental/WaterAdvanced/WaterAdvanced.uplugin
+++ b/Engine/Plugins/Experimental/WaterAdvanced/WaterAdvanced.uplugin
@@ -3,17 +3,17 @@
 	"Version": 1,
 	"VersionName": "1.0",
 	"FriendlyName": "WaterAdvanced",
-	"Description": "",
+	"Description": "Collection of easy to use water simulation systems built on the Niagara Fluids and Water plugins",
 	"Category": "Water",
-	"CreatedBy": "",
-	"CreatedByURL": "",
+	"CreatedBy": "Epic Games, Inc.",
+	"CreatedByURL": "https://epicgames.com",
 	"DocsURL": "",
 	"MarketplaceURL": "",
 	"SupportURL": "",
 	"EnabledByDefault": false,
 	"CanContainContent": true,
-	"IsBetaVersion": false,
-	"IsExperimentalVersion": true,
+	"IsBetaVersion": true,
+	"IsExperimentalVersion": false,
 	"Installed": false,
 	"Modules": [
 		{
