support ripples when moving between water zones

diff --git a/Engine/Plugins/Experimental/WaterAdvanced/Source/WaterAdvanced/Private/ShallowWaterSubsystem.cpp b/Engine/Plugins/Experimental/WaterAdvanced/Source/WaterAdvanced/Private/ShallowWaterSubsystem.cpp
--- a/Engine/Plugins/Experimental/WaterAdvanced/Source/WaterAdvanced/Private/ShallowWaterSubsystem.cpp
+++ b/Engine/Plugins/Experimental/WaterAdvanced/Source/WaterAdvanced/Private/ShallowWaterSubsystem.cpp
@@ -187,11 +187,6 @@ void UShallowWaterSubsystem::Tick(float DeltaTime)
 	{
 		UpdateOverlappingWaterBodiesHistory(WaterBodies.Array());
 	
-		if (WaterInfoTexture == nullptr)
-		{
-			TryGetOrWaitForWaterInfoTextureFromWaterBodies(WaterBodies);
-		}
-	
 		for (AWaterBody* Water : WaterBodies)
 		{
 			if (Water == nullptr)
@@ -492,6 +487,8 @@ void UShallowWaterSubsystem::InitializeParameters()
 		ShallowWaterNiagaraSimulation->SetVariableVec2(FName("WorldGridSize"), FVector2D(GetGridSize()));
 		ShallowWaterNiagaraSimulation->SetVariableInt(FName("ResolutionMaxAxis"), GetGridResolution());
 		ShallowWaterNiagaraSimulation->SetVariableTextureRenderTarget(FName("NormalRT"), NormalRT);		
+
+		ShallowWaterNiagaraSimulation->SetVariableBool(FName("UseBakedSim"), false);
 	}
 	else
 	{
@@ -536,8 +533,8 @@ void UShallowWaterSubsystem::UpdateGridMovement()
 
 	float BestDistanceSqr = FLT_MAX;
 	TOptional<FVector> BestWaterLocation;
-	
-	TWeakObjectPtr<AWaterBody> BestWaterBody;
+	TWeakObjectPtr<AWaterBody> NewBestWaterBody;
+
 	for (TWeakObjectPtr<AWaterBody> WeakWaterBody : LastOverlappingWaterBodies_Internal)
 	{
 		if (const AWaterBody* WaterBody = WeakWaterBody.Get())
@@ -556,7 +553,7 @@ void UShallowWaterSubsystem::UpdateGridMovement()
 			{
 				BestDistanceSqr = ds;
 				BestWaterLocation = WaterLocation;
-				BestWaterBody = WeakWaterBody;
+				NewBestWaterBody = WeakWaterBody;
 			}
 		}
 	}
@@ -566,6 +563,10 @@ void UShallowWaterSubsystem::UpdateGridMovement()
 		return;
 	}
 
+	// Check to see if the best water body has changed
+	bool BestWaterBodyChanged = NewBestWaterBody != BestWaterBody;
+	BestWaterBody = NewBestWaterBody;
+
 	const FVector ProjectedLocation = FVector(CursorPawnLocation.X, CursorPawnLocation.Y, 0);
 	ShallowWaterNiagaraSimulation->SetWorldLocation(ProjectedLocation);
 	ShallowWaterNiagaraSimulation->SetVariableFloat(FName("WaterHeightAtParent"), BestWaterLocation.GetValue().Z);
@@ -576,11 +577,11 @@ void UShallowWaterSubsystem::UpdateGridMovement()
 	ShallowWaterNiagaraSimulation->SetVariableFloat(FName("SecondsUntilDestroyed"), SecondsUntilDestroyed);
 
 	ShallowWaterNiagaraSimulation->SetVariableBool(FName("UseDebugRender"), GSWDebugRender == 1);
-	ShallowWaterNiagaraSimulation->SetVariableBool(FName("UseWaterInfoTexture"), GSWUseWaterInfoTexture == 1);
-
-	bool UseBakedSim = false;
-	if (BestWaterBody != nullptr && BestWaterBody->GetWaterBodyComponent())
+	
+	if (BestWaterBodyChanged && BestWaterBody != nullptr)
 	{
+		bool UseBakedSim = false;
+
 		if (const UWaterBodyComponent* WaterBodyComp = BestWaterBody->GetWaterBodyComponent())
 		{
 			if (WaterBodyComp->UseBakedSimulationForQueriesAndPhysics())
@@ -598,9 +599,18 @@ void UShallowWaterSubsystem::UpdateGridMovement()
 				}
 			}
 		}
+		
+		ShallowWaterNiagaraSimulation->SetVariableBool(FName("UseBakedSim"), UseBakedSim);
+
+		// update to use the water info from this water body
+		bool UseWaterInfoTexture = GSWUseWaterInfoTexture == 1;
+		ShallowWaterNiagaraSimulation->SetVariableBool(FName("UseWaterInfoTexture"), UseWaterInfoTexture);
+		if (UseWaterInfoTexture)
+		{			
+			TryGetOrWaitForWaterInfoTextureFromWaterBody(BestWaterBody.Get());
+		}
 	}
-	ShallowWaterNiagaraSimulation->SetVariableBool(FName("UseBakedSim"), UseBakedSim);
-
+	
 #if ENABLE_DRAW_DEBUG
 	if (GSWDrawWaterSurfaceProjection)
 	{
@@ -1231,6 +1241,58 @@ TArray<APawn*> UShallowWaterSubsystem::GetPawnsInRange(const FVector ObservingLo
 	return Results;
 }
 
+void UShallowWaterSubsystem::SetSimParametersForWaterZone(AWaterZone *WaterZone)
+{
+	// The following index assume that there is no split screen support and will request the position of the first player's water view.
+	const int32 PlayerIndex = 0;
+	FVector ZoneLocation;
+	WaterZone->GetDynamicWaterInfoCenter(PlayerIndex, ZoneLocation);
+
+	const FVector2D ZoneExtent = FVector2D(WaterZone->GetDynamicWaterInfoExtent());
+	const FVector2D WaterHeightExtents = FVector2D(WaterZone->GetWaterHeightExtents());
+	const float GroundZMin = WaterZone->GetGroundZMin();
+
+	ShallowWaterNiagaraSimulation->SetVariableVec2(FName("WaterZoneLocation"), FVector2D(ZoneLocation));
+	ShallowWaterNiagaraSimulation->SetVariableVec2(FName("WaterZoneExtent"), ZoneExtent);
+	ShallowWaterNiagaraSimulation->SetVariableInt(FName("WaterZoneIdx"), WaterZone->GetWaterZoneIndex());
+}
+
+bool UShallowWaterSubsystem::TryGetOrWaitForWaterInfoTextureFromWaterBody(AWaterBody *CurrentWaterBody)
+{    
+	if (const UWaterBodyComponent* WaterBodyComp = CurrentWaterBody->GetWaterBodyComponent())
+	{
+	    if (AWaterZone* WaterZone = WaterBodyComp->GetWaterZone())
+	    {
+	    	const TObjectPtr<UTextureRenderTarget2DArray> NewWaterInfoTexture = WaterZone->WaterInfoTextureArray;
+
+	    	if (NewWaterInfoTexture == nullptr)
+	    	{
+	    		WaterZone->GetOnWaterInfoTextureArrayCreated().RemoveDynamic(this, &UShallowWaterSubsystem::OnWaterInfoTextureArrayCreated);
+	    		WaterZone->GetOnWaterInfoTextureArrayCreated().AddDynamic(this, &UShallowWaterSubsystem::OnWaterInfoTextureArrayCreated);
+	    	}
+	    	else
+	    	{
+	    		OnWaterInfoTextureArrayCreated(NewWaterInfoTexture);
+	    	}
+
+			if (ShallowWaterNiagaraSimulation)
+			{
+				SetSimParametersForWaterZone(WaterZone);
+			}
+			else
+			{
+				ensureMsgf(false, TEXT("UShallowWaterSubsystem::TryGetOrWaitForWaterInfoTextureFromWaterBodies was called with NULL ShallowWaterNiagaraSimulation"));
+				return false;
+			}
+	    	
+			// Currently there can only be one unique WaterInfoTexture
+	    	return true;
+	    }
+	}   
+
+	return false;
+}
+
 void UShallowWaterSubsystem::TryGetOrWaitForWaterInfoTextureFromWaterBodies(TSet<AWaterBody*> CurrentWaterBodies)
 {
     for (const AWaterBody* CurrentWaterBody : CurrentWaterBodies)
@@ -1253,24 +1315,14 @@ void UShallowWaterSubsystem::TryGetOrWaitForWaterInfoTextureFromWaterBodies(TSet
 
 				if (ShallowWaterNiagaraSimulation)
 				{
-					// The following index assume that there is no split screen support and will request the position of the first player's water view.
-					const int32 PlayerIndex = 0;
-					FVector ZoneLocation;
-					WaterZone->GetDynamicWaterInfoCenter(PlayerIndex, ZoneLocation);
-
-					const FVector2D ZoneExtent = FVector2D(WaterZone->GetDynamicWaterInfoExtent());
-					const FVector2D WaterHeightExtents = FVector2D(WaterZone->GetWaterHeightExtents());
-					const float GroundZMin = WaterZone->GetGroundZMin();
-
-					ShallowWaterNiagaraSimulation->SetVariableVec2(FName("WaterZoneLocation"), FVector2D(ZoneLocation));
-					ShallowWaterNiagaraSimulation->SetVariableVec2(FName("WaterZoneExtent"), ZoneExtent);
-					ShallowWaterNiagaraSimulation->SetVariableInt(FName("WaterZoneIdx"), WaterZone->GetWaterZoneIndex());
+					SetSimParametersForWaterZone(WaterZone);
 				}
 				else
 				{
 					ensureMsgf(false, TEXT("UShallowWaterSubsystem::TryGetOrWaitForWaterInfoTextureFromWaterBodies was called with NULL ShallowWaterNiagaraSimulation"));
 					return;
 				}
+
 	    		// Currently there can only be one unique WaterInfoTexture
 	    		break;
 	    	}
diff --git a/Engine/Plugins/Experimental/WaterAdvanced/Source/WaterAdvanced/Public/ShallowWaterSubsystem.h b/Engine/Plugins/Experimental/WaterAdvanced/Source/WaterAdvanced/Public/ShallowWaterSubsystem.h
--- a/Engine/Plugins/Experimental/WaterAdvanced/Source/WaterAdvanced/Public/ShallowWaterSubsystem.h
+++ b/Engine/Plugins/Experimental/WaterAdvanced/Source/WaterAdvanced/Public/ShallowWaterSubsystem.h
@@ -228,6 +228,8 @@ protected:
 	/*
 	 * WaterBody
 	 */
+	UE_API void SetSimParametersForWaterZone(AWaterZone* WaterZone);
+	UE_API bool TryGetOrWaitForWaterInfoTextureFromWaterBody(AWaterBody* CurrentWaterBody);
 	UE_API void TryGetOrWaitForWaterInfoTextureFromWaterBodies(TSet<AWaterBody*> CurrentWaterBodies);
 	UFUNCTION()
 	UE_API void OnWaterInfoTextureArrayCreated(const UTextureRenderTarget2DArray* WaterInfoTexture);
@@ -253,6 +255,7 @@ protected:
 	bool bTickCacheValid = false;
 	APawn* CachedCursorPawn = nullptr;
 	TArray<APawn*> CachedPawnsInRange;
+	TWeakObjectPtr<AWaterBody> BestWaterBody;
 
 private:
 	float LastTimeOverlappingAnyWaterBody = -FLT_MAX;
