Fix Garbage References
(Authored by [REDACTED_PERSON])

diff --git a/Engine/Plugins/Experimental/WaterAdvanced/Source/WaterAdvanced/Private/ShallowWaterSubsystem.cpp b/Engine/Plugins/Experimental/WaterAdvanced/Source/WaterAdvanced/Private/ShallowWaterSubsystem.cpp
--- a/Engine/Plugins/Experimental/WaterAdvanced/Source/WaterAdvanced/Private/ShallowWaterSubsystem.cpp
+++ b/Engine/Plugins/Experimental/WaterAdvanced/Source/WaterAdvanced/Private/ShallowWaterSubsystem.cpp
@@ -997,7 +997,7 @@ void UShallowWaterSubsystem::CleanUpVehicleCollisionProxies()
 	{
 		// Shotgun approach to cover all weird possibilities
 		// E.g. if the Vehicle component is destroyed but not the owning actor somehow
-		if ( Entry.Key.Component == nullptr || Entry.Key.Component->IsBeingDestroyed() ||
+		if ( Entry.Key.WeakComponent.Get() == nullptr ||
 			Entry.Value == nullptr || Entry.Value->IsBeingDestroyed())
 		{
 			EntriesToRemove.Add(Entry);
@@ -1037,7 +1037,8 @@ void UShallowWaterSubsystem::UpdateCollisionForPendingContexts()
 
 void UShallowWaterSubsystem::EnableCollisionForContext(const FShallowWaterCollisionContext& Context)
 {
-	if (Context.Component == nullptr)
+	USkeletalMeshComponent* Component = Context.WeakComponent.Get();
+	if (Component == nullptr)
 	{		
 		UE_LOG(LogShallowWater, Warning, TEXT("EnableCollisionForContext() - Context Component is nullptr"));
 		return;
@@ -1046,7 +1047,7 @@ void UShallowWaterSubsystem::EnableCollisionForContext(const FShallowWaterCollis
 	{
 	case EShallowWaterCollisionContextType::Pawn:
 		
-		Context.Component->ComponentTags.AddUnique(ColliderComponentTag);
+		Component->ComponentTags.AddUnique(ColliderComponentTag);
 		break;
 		
 	case EShallowWaterCollisionContextType::Vehicle:
@@ -1054,8 +1055,8 @@ void UShallowWaterSubsystem::EnableCollisionForContext(const FShallowWaterCollis
 		// Spawn proxy SKM, attach to vehicle
 		const FString BaseName = TEXT("FluidsimCollisionProxy");
 		const FName CompName = MakeUniqueObjectName(this, USkeletalMeshComponent::StaticClass(), FName(*BaseName));
-		USkeletalMeshComponent* ProxyComp = NewObject<USkeletalMeshComponent>(Context.Component->GetOwner(), USkeletalMeshComponent::StaticClass(), CompName);
-		ProxyComp->SetSkeletalMeshAsset(Context.Component->GetSkeletalMeshAsset());
+		USkeletalMeshComponent* ProxyComp = NewObject<USkeletalMeshComponent>(Component->GetOwner(), USkeletalMeshComponent::StaticClass(), CompName);
+		ProxyComp->SetSkeletalMeshAsset(Component->GetSkeletalMeshAsset());
 
 		// Spawn an empty dummy SKM component to be used as Collision 
 		// Apply PhysicsAsset override if defined in the data asset
@@ -1082,7 +1083,7 @@ void UShallowWaterSubsystem::EnableCollisionForContext(const FShallowWaterCollis
 				TSoftObjectPtr<UPhysicsAsset> TmpPhysicsAsset = FoundOverride->PhysicsAsset;
 
 				UAssetManager::GetStreamableManager().RequestAsyncLoad(TmpPhysicsAsset.ToSoftObjectPath(),
-					FStreamableDelegate::CreateWeakLambda(this, [this, TmpPhysicsAsset, ProxyComp, Context]()
+					FStreamableDelegate::CreateWeakLambda(this, [this, TmpPhysicsAsset, ProxyComp, Context, Component]()
 						{								
 							if (TmpPhysicsAsset.IsValid())
 							{								
@@ -1092,9 +1093,9 @@ void UShallowWaterSubsystem::EnableCollisionForContext(const FShallowWaterCollis
 
 								ProxyComp->ComponentTags.AddUnique(ColliderComponentTag);
 
-								ProxyComp->SetupAttachment(Context.Component->GetOwner()->GetRootComponent());
-								ProxyComp->SetWorldTransform(Context.Component->GetComponentTransform());
-								ProxyComp->SetLeaderPoseComponent(Context.Component);
+								ProxyComp->SetupAttachment(Component->GetOwner()->GetRootComponent());
+								ProxyComp->SetWorldTransform(Component->GetComponentTransform());
+								ProxyComp->SetLeaderPoseComponent(Component);
 								ProxyComp->SetVisibility(false);
 								ProxyComp->SetCollisionEnabled(ECollisionEnabled::NoCollision);
 								ProxyComp->RegisterComponent();
@@ -1109,7 +1110,7 @@ void UShallowWaterSubsystem::EnableCollisionForContext(const FShallowWaterCollis
 		{
 			if (GSWUseFullVehiclePhysicsAssets)
 			{
-				Context.Component->ComponentTags.AddUnique(ColliderComponentTag);
+				Component->ComponentTags.AddUnique(ColliderComponentTag);
 			}
 			else
 			{
@@ -1145,7 +1146,10 @@ void UShallowWaterSubsystem::DisableCollisionForContext(const FShallowWaterColli
 	switch (Context.Type)
 	{
 	case EShallowWaterCollisionContextType::Pawn:
-		Context.Component->ComponentTags.Remove(ColliderComponentTag);
+		if (USkeletalMeshComponent* Component = Context.WeakComponent.Get())
+		{
+			Component->ComponentTags.Remove(ColliderComponentTag);
+		}
 		break;
 	case EShallowWaterCollisionContextType::Vehicle:
 		DisableCollisionForVehicle(Context);
diff --git a/Engine/Plugins/Experimental/WaterAdvanced/Source/WaterAdvanced/Public/ShallowWaterSubsystem.h b/Engine/Plugins/Experimental/WaterAdvanced/Source/WaterAdvanced/Public/ShallowWaterSubsystem.h
--- a/Engine/Plugins/Experimental/WaterAdvanced/Source/WaterAdvanced/Public/ShallowWaterSubsystem.h
+++ b/Engine/Plugins/Experimental/WaterAdvanced/Source/WaterAdvanced/Public/ShallowWaterSubsystem.h
@@ -40,23 +40,25 @@ struct FShallowWaterCollisionContext
 	FShallowWaterCollisionContext() {};
 	FShallowWaterCollisionContext(EShallowWaterCollisionContextType InType, USkeletalMeshComponent* InComponent)
 		: Type(InType)
-		, Component(InComponent)
+		, WeakComponent(InComponent)
 	{
 		if(InComponent == nullptr)
 		{
 			ensureMsgf(false, TEXT("FShallowWaterCollisionContext constructor was given nullptr component as input"));
 			return;
 		}
-		UniqueID = Component->GetUniqueID();
+		UniqueID = InComponent->GetUniqueID();
 	}
+
 	EShallowWaterCollisionContextType Type = EShallowWaterCollisionContextType::Pawn;
+
 	UPROPERTY()
-	TObjectPtr<USkeletalMeshComponent> Component = nullptr;
+	TWeakObjectPtr<USkeletalMeshComponent> WeakComponent = nullptr;
 
 	// #todo Theoretically not enough as unique identification since Component->GetUniqueID() is "reused so it is only unique while the object is alive"
 	uint32 UniqueID = 0;
 
-	bool IsValidAndAlive() const { return Component != nullptr && !Component->IsBeingDestroyed(); }
+	bool IsValidAndAlive() const { return WeakComponent.Get() != nullptr; }
 
 	bool operator==(const FShallowWaterCollisionContext& Other) const
 	{
@@ -235,7 +237,7 @@ protected:
 	UE_API void OnWaterInfoTextureArrayCreated(const UTextureRenderTarget2DArray* WaterInfoTexture);
 
 	UPROPERTY(BlueprintReadOnly, VisibleAnywhere, Category="Shallow Water")
-	TObjectPtr<const UTextureRenderTarget2DArray> WaterInfoTexture;
+	TWeakObjectPtr<const UTextureRenderTarget2DArray> WaterInfoTexture;
 	UPROPERTY(BlueprintReadOnly, VisibleAnywhere, Category="Shallow Water")
 	TObjectPtr<UTextureRenderTarget2D> NormalRT;
 
