- Fix an asan issue using the component in the lambda
- Bring changes over for garbage references also

diff --git a/Engine/Plugins/Experimental/WaterAdvanced/Source/WaterAdvanced/Private/ShallowWaterSubsystem.cpp b/Engine/Plugins/Experimental/WaterAdvanced/Source/WaterAdvanced/Private/ShallowWaterSubsystem.cpp
--- a/Engine/Plugins/Experimental/WaterAdvanced/Source/WaterAdvanced/Private/ShallowWaterSubsystem.cpp
+++ b/Engine/Plugins/Experimental/WaterAdvanced/Source/WaterAdvanced/Private/ShallowWaterSubsystem.cpp
@@ -1085,8 +1085,9 @@ void UShallowWaterSubsystem::EnableCollisionForContext(const FShallowWaterCollis
 				UAssetManager::GetStreamableManager().RequestAsyncLoad(TmpPhysicsAsset.ToSoftObjectPath(),
 					FStreamableDelegate::CreateWeakLambda(this, [this, TmpPhysicsAsset, ProxyComp, Context, Component]()
 						{								
-							if (TmpPhysicsAsset.IsValid())
-							{								
+							USkeletalMeshComponent* Component = Context.WeakComponent.Get();
+							if (Component && Component->GetOwner() && TmpPhysicsAsset.IsValid())
+							{
 								UPhysicsAsset* PhysicsAssetOverride = TmpPhysicsAsset.Get();
 
 								ProxyComp->SetPhysicsAsset(PhysicsAssetOverride);
