- Fix CIS shadowing issue on some platforms

#rnx

diff --git a/Engine/Plugins/Experimental/WaterAdvanced/Source/WaterAdvanced/Private/ShallowWaterSubsystem.cpp b/Engine/Plugins/Experimental/WaterAdvanced/Source/WaterAdvanced/Private/ShallowWaterSubsystem.cpp
--- a/Engine/Plugins/Experimental/WaterAdvanced/Source/WaterAdvanced/Private/ShallowWaterSubsystem.cpp
+++ b/Engine/Plugins/Experimental/WaterAdvanced/Source/WaterAdvanced/Private/ShallowWaterSubsystem.cpp
@@ -1085,8 +1085,8 @@ void UShallowWaterSubsystem::EnableCollisionForContext(const FShallowWaterCollis
 				UAssetManager::GetStreamableManager().RequestAsyncLoad(TmpPhysicsAsset.ToSoftObjectPath(),
 					FStreamableDelegate::CreateWeakLambda(this, [this, TmpPhysicsAsset, ProxyComp, Context, Component]()
 						{								
-							USkeletalMeshComponent* Component = Context.WeakComponent.Get();
-							if (Component && Component->GetOwner() && TmpPhysicsAsset.IsValid())
+							USkeletalMeshComponent* ContextComponent = Context.WeakComponent.Get();
+							if (ContextComponent && ContextComponent->GetOwner() && TmpPhysicsAsset.IsValid())
 							{
 								UPhysicsAsset* PhysicsAssetOverride = TmpPhysicsAsset.Get();
 
@@ -1094,9 +1094,9 @@ void UShallowWaterSubsystem::EnableCollisionForContext(const FShallowWaterCollis
 
 								ProxyComp->ComponentTags.AddUnique(ColliderComponentTag);
 
-								ProxyComp->SetupAttachment(Component->GetOwner()->GetRootComponent());
-								ProxyComp->SetWorldTransform(Component->GetComponentTransform());
-								ProxyComp->SetLeaderPoseComponent(Component);
+								ProxyComp->SetupAttachment(ContextComponent->GetOwner()->GetRootComponent());
+								ProxyComp->SetWorldTransform(ContextComponent->GetComponentTransform());
+								ProxyComp->SetLeaderPoseComponent(ContextComponent);
 								ProxyComp->SetVisibility(false);
 								ProxyComp->SetCollisionEnabled(ECollisionEnabled::NoCollision);
 								ProxyComp->RegisterComponent();
