* virtual textures - cached sims are now stored in virtual textures to greatly improve memory usage.
* remove normal smoothing and cache processing from niagara system
* disable cpu throttling when simulating in editor
* many other small fixes

#virtualized

diff --git a/Engine/Plugins/Experimental/WaterAdvanced/Source/WaterAdvanced/Private/ShallowWaterRiverActor.cpp b/Engine/Plugins/Experimental/WaterAdvanced/Source/WaterAdvanced/Private/ShallowWaterRiverActor.cpp
--- a/Engine/Plugins/Experimental/WaterAdvanced/Source/WaterAdvanced/Private/ShallowWaterRiverActor.cpp
+++ b/Engine/Plugins/Experimental/WaterAdvanced/Source/WaterAdvanced/Private/ShallowWaterRiverActor.cpp
@@ -35,6 +35,11 @@
 
 #include "EngineUtils.h"
 
+#if WITH_EDITOR
+#include "Editor.h"
+#endif
+
+
 #include UE_INLINE_GENERATED_CPP_BY_NAME(ShallowWaterRiverActor)
 
 bool bShallowWaterRiverDebugVisualize = false;
@@ -128,10 +133,54 @@ FBoxSphereBounds UShallowWaterRiverComponent::InitializeCaptureDI(const FName& D
 	return FBoxSphereBounds(BottomContourCombinedWorldBoundsBuilder);	
 }
 
+void UShallowWaterRiverComponent::ConvertToVirtualTextures()
+{
+#if WITH_EDITOR
+	bool HasChanged = false;
+
+	bUseVirtualTextures = true;
+
+	if (BakedWaterSurfaceTexture != NULL && !BakedWaterSurfaceTexture->VirtualTextureStreaming)
+	{	
+		SimRes = FVector2D(BakedWaterSurfaceTexture->Source.GetSizeX(), BakedWaterSurfaceTexture->Source.GetSizeY());
+		RiverSimSystem->SetVariableVec2(FName("SimRes"), SimRes);
+
+		InitializeVirtualTexture(BakedWaterSurfaceTexture);
+		UE_LOG(LogShallowWater, Warning, TEXT("Baked water surface texture was not virtual- converting.  Recommended resave."));
+
+		HasChanged = true;
+	}
+	
+	if (BakedWaterSurfaceNormalTexture != NULL && !BakedWaterSurfaceNormalTexture->VirtualTextureStreaming)
+	{
+		InitializeVirtualTexture(BakedWaterSurfaceNormalTexture);
+		UE_LOG(LogShallowWater, Warning, TEXT("Baked water surface normal texture was not virtual- converting.  Recommended resave."));
+
+		HasChanged = true;
+	}
+
+	if (BakedFoamTexture != NULL && !BakedFoamTexture->VirtualTextureStreaming)
+	{
+		InitializeVirtualTexture(BakedFoamTexture);
+		UE_LOG(LogShallowWater, Warning, TEXT("Baked foam texture was not virtual- converting.  Recommended resave."));
+
+		HasChanged = true;
+	}
+
+	if (HasChanged)
+	{
+		PostEditChange();
+	}
+#endif
+}
+
 void UShallowWaterRiverComponent::PostLoad()
 {
 	Super::PostLoad();	
 
+	// ensure all baked textures are virtual for backwards compatibility
+	ConvertToVirtualTextures();
+
 	if (RenderState == EShallowWaterRenderState::LiveSim || RiverSimSystem == nullptr)
 	{
 	#if WITH_EDITOR
@@ -485,12 +534,22 @@ void UShallowWaterRiverComponent::Rebuild()
 	RiverSimSystem->SetVariableVec2(FName("WorldGridSize"), WorldGridSize);
 	RiverSimSystem->SetVariableInt(FName("ResolutionMaxAxis"), ResolutionMaxAxis);
 
+	// #todo(dmp): would be better to initialize the user var inside of niagara rather than recomputing res here
+	SimRes = FVector2D(ResolutionMaxAxis, ResolutionMaxAxis * WorldGridSize.Y / WorldGridSize.X);
+	if (WorldGridSize.Y > WorldGridSize.X)
+	{
+		SimRes = FVector2D(ResolutionMaxAxis * WorldGridSize.X / WorldGridSize.Y, ResolutionMaxAxis);
+	}
+	RiverSimSystem->SetVariableVec2(FName("SimRes"), SimRes);
+
 	// pad out source's box height a so it intersects the sim plane.  This value doesn't matter much so we hardcode it
 	float Overshoot = 1000.f;
 	float FinalSourceHeight = 2. * CombinedBounds.BoxExtent.Z + Overshoot;
 
-	// Get sources
-	int32 i = 0;
+	// Get sources	
+	TArray<FVector> SourcePosArray;
+	TArray<FVector3f> FullSourceSizeArray;
+	TArray<float> FullSourceAngleArray;
 	for (TSoftObjectPtr<AWaterBody> CurrWaterBody : SourceRiverWaterBodies)
 	{
 		FVector CurrSourcePos;
@@ -529,20 +588,26 @@ void UShallowWaterRiverComponent::Rebuild()
 			FullSourceAngle *= -1;
 		}
 
-		UNiagaraDataInterfaceArrayFunctionLibrary::SetNiagaraArrayPositionValue(RiverSimSystem, "User.SourcePosArray", i, FullSourcePos, true);
-		UNiagaraDataInterfaceArrayFunctionLibrary::SetNiagaraArrayVectorValue(RiverSimSystem, "User.SourceSizeArray", i, FullSourceSize, true);
-		UNiagaraDataInterfaceArrayFunctionLibrary::SetNiagaraArrayFloatValue(RiverSimSystem, "User.SourceAngleArray", i, FullSourceAngle, true);
-		i++;
+
+		SourcePosArray.Add(FullSourcePos);
+		FullSourceSizeArray.Add(FVector3f(FullSourceSize));
+		FullSourceAngleArray.Add(FullSourceAngle);				
 	}
+	UNiagaraDataInterfaceArrayFunctionLibrary::SetNiagaraArrayPosition(RiverSimSystem, "User.SourcePosArray",SourcePosArray);
+	UNiagaraDataInterfaceArrayFunctionLibrary::SetNiagaraArrayVector(RiverSimSystem, "User.SourceSizeArray", FullSourceSizeArray);
+	UNiagaraDataInterfaceArrayFunctionLibrary::SetNiagaraArrayFloat(RiverSimSystem, "User.SourceAngleArray", FullSourceAngleArray);	
+
+	// get sinks		
+	TArray<FVector> SinkPosArray;
+	TArray<FVector3f> FullSinkSizeArray;
+	TArray<float> FullSinkAngleArray;
 
-	// get sinks	
-	FVector SinkPos(0, 0, 0);
-	float SinkWidth = 1;
-	float SinkDepth = 1;
-	FVector SinkDir(1, 0, 0);
-	
 	for (TSoftObjectPtr<AWaterBody> CurrWaterBody : SinkRiverWaterBodies)
 	{
+		FVector SinkPos(0, 0, 0);
+		float SinkWidth = 1;
+		float SinkDepth = 1;
+		FVector SinkDir(1, 0, 0);
 		if (!QueryWaterAtSplinePoint(CurrWaterBody, -1, SinkPos, SinkDir, SinkWidth, SinkDepth))
 		{
 			UE_LOG(LogShallowWater, Warning, TEXT("UShallowWaterRiverComponent::Rebuild() - water sink query failed"));
@@ -559,7 +624,6 @@ void UShallowWaterRiverComponent::Rebuild()
 		FVector BaseVector = {0,1,0};	
 		double FullSinkAngle =  FMath::Acos(FVector::DotProduct(BaseVector, SinkDir));
 		
-
 		FVector AxisToUse = FVector::CrossProduct(BaseVector, SinkDir);
 		AxisToUse.Normalize();
 
@@ -576,14 +640,14 @@ void UShallowWaterRiverComponent::Rebuild()
 		{
 			FullSinkAngle *= -1;
 		}
-
-		UNiagaraDataInterfaceArrayFunctionLibrary::SetNiagaraArrayPositionValue(RiverSimSystem, "User.SinkPosArray", i, SinkPos, true);
-
-		UNiagaraDataInterfaceArrayFunctionLibrary::SetNiagaraArrayVectorValue(RiverSimSystem, "User.SinkSizeArray", i, FullSinkSize, true);
-		UNiagaraDataInterfaceArrayFunctionLibrary::SetNiagaraArrayFloatValue(RiverSimSystem, "User.SinkAngleArray", i, FullSinkAngle, true);		
-
-		i++;
+		
+		SinkPosArray.Add(SinkPos);
+		FullSinkSizeArray.Add(FVector3f(FullSinkSize));
+		FullSinkAngleArray.Add(FullSinkAngle);				
 	}
+	UNiagaraDataInterfaceArrayFunctionLibrary::SetNiagaraArrayPosition(RiverSimSystem, "User.SinkPosArray",SinkPosArray);
+	UNiagaraDataInterfaceArrayFunctionLibrary::SetNiagaraArrayVector(RiverSimSystem, "User.SinkSizeArray", FullSinkSizeArray);
+	UNiagaraDataInterfaceArrayFunctionLibrary::SetNiagaraArrayFloat(RiverSimSystem, "User.SinkAngleArray", FullSinkAngleArray);		
 
 	RiverSimSystem->SetVariableFloat(FName("SimSpeed"), SimSpeed);
 	RiverSimSystem->SetVariableInt(FName("NumSteps"), NumSteps);
@@ -629,6 +693,22 @@ void UShallowWaterRiverComponent::Rebuild()
 	NormalDetailRT = OceanPatchNormalRT;
 	RiverSimSystem->SetVariableTextureRenderTarget(FName("NormalDetailRT"), NormalDetailRT);
 
+#if WITH_EDITOR
+	// run the live sim and make sure that we enable/disable cpu throttling	
+	if (!bReadBakedSim)
+	{		
+		GEditor->ShouldDisableCPUThrottlingDelegates.Add(UEditorEngine::FShouldDisableCPUThrottling::CreateUObject(this, &UShallowWaterRiverComponent::ShouldDisableCPUThrottling));
+		ShouldDisableCPUThrottlingDelegateHandle = GEditor->ShouldDisableCPUThrottlingDelegates.Last().GetHandle();
+	}
+	else
+	{
+		GEditor->ShouldDisableCPUThrottlingDelegates.RemoveAll([this](const UEditorEngine::FShouldDisableCPUThrottling& Delegate)
+		{
+			return Delegate.GetHandle() == ShouldDisableCPUThrottlingDelegateHandle;
+		});
+	}
+#endif
+
 	bIsInitialized = true;
 }
 
@@ -711,6 +791,12 @@ void UShallowWaterRiverComponent::Bake()
 		return;
 	}
 
+	if (RenderState != EShallowWaterRenderState::LiveSim)
+	{
+		UE_LOG(LogShallowWater, Warning, TEXT("UShallowWaterRiverComponent::Bake() - Must be in live sim mode to bake"));
+		return;
+	}
+
 	if (SourceRiverWaterBodies.Num() != 0)
 	{
 		for (TSoftObjectPtr<AWaterBody > CurrWaterBody : SourceRiverWaterBodies)
@@ -721,10 +807,27 @@ void UShallowWaterRiverComponent::Bake()
 				return;
 			}
 		}
+	}	
+
+	SimRes = FVector2D(BakedWaterSurfaceRT->SizeX, BakedWaterSurfaceRT->SizeY);
+	RiverSimSystem->SetVariableVec2(FName("SimRes"), SimRes);
+
+	if (WorldGridSize.X > WorldGridSize.Y)
+	{
+		SimDx = WorldGridSize.X / ResolutionMaxAxis;
 	}
+	else
+	{
+		SimDx = WorldGridSize.Y / ResolutionMaxAxis;
+	}		
 
 	BakedWaterSurfaceTexture = BakedWaterSurfaceRT->ConstructTexture2D(this, "BakedRiverTexture", TextureObjectFlags);
 
+	if (bUseVirtualTextures)
+	{
+		InitializeVirtualTexture(BakedWaterSurfaceTexture);
+	}
+
 	RiverSimSystem->SetVariableTexture(FName("BakedSimTexture"), BakedWaterSurfaceTexture);
 
 	// Readback to get the river texture values as an array
@@ -750,14 +853,26 @@ void UShallowWaterRiverComponent::Bake()
 		FloatVal.W = WaterVelocity.Y;
 
 		ShallowWaterSimArrayValues[Index++] = FloatVal;
-	}
+	}	
 
 	// bake foam and other data to texture
 	BakedFoamTexture = BakedFoamRT->ConstructTexture2D(this, "BakedFoamTexture", TextureObjectFlags);
+	
+	if (bUseVirtualTextures)
+	{		
+		InitializeVirtualTexture(BakedFoamTexture);
+	}
+	
 	RiverSimSystem->SetVariableTexture(FName("BakedFoamTexture"), BakedFoamTexture);
 	
 	// bake normal to texture
 	BakedWaterSurfaceNormalTexture = BakedWaterSurfaceNormalRT->ConstructTexture2D(this, "BakedWaterSurfaceNormalTexture", TextureObjectFlags);
+	
+	if (bUseVirtualTextures)
+	{
+		InitializeVirtualTexture(BakedWaterSurfaceNormalTexture);
+	}
+	
 	RiverSimSystem->SetVariableTexture(FName("BakedWaterSurfaceNormalTexture"), BakedWaterSurfaceNormalTexture);
 
 	// clear references to old baked sim on water body actors
@@ -902,6 +1017,19 @@ void UShallowWaterRiverComponent::Bake()
 	}
 }
 
+void UShallowWaterRiverComponent::InitializeVirtualTexture(TObjectPtr<UTexture2D> InTexture)
+{	
+	InTexture->Modify();
+	InTexture->MipGenSettings = TextureMipGenSettings::TMGS_SimpleAverage;
+	InTexture->PowerOfTwoMode = ETexturePowerOfTwoSetting::PadToPowerOfTwo;
+	InTexture->VirtualTextureStreaming = true;
+
+	InTexture->UpdateResource();	
+	InTexture->WaitForStreaming(true, true);	
+	InTexture->BlockOnAnyAsyncBuild();
+	InTexture->PostEditChange();
+}
+
 bool UShallowWaterRiverComponent::QueryWaterAtSplinePoint(TSoftObjectPtr<AWaterBody> WaterBody, int SplinePoint, FVector& OutPos, FVector& OutTangent, float& OutWidth, float& OutDepth)
 {	
 	if (WaterBody != nullptr)
@@ -1050,6 +1178,8 @@ void UShallowWaterRiverComponent::UpdateRenderState()
 
 	if (RiverSimSystem != nullptr)
 	{		
+		RiverSimSystem->SetVariableVec2(FName("SimRes"), SimRes);
+
 		RiverSimSystem->SetVariableBool(FName("RenderWaterSurface"), !RenderWaterBody);
 		RiverSimSystem->SetVariableBool(FName("RenderSecondary"), RenderSecondary);
 		RiverSimSystem->SetVariableBool(FName("DebugRenderBottomContour"), RenderState == EShallowWaterRenderState::DebugRenderBottomContour);
@@ -1101,7 +1231,7 @@ void UShallowWaterRiverComponent::UpdateRenderState()
 					WaterInfoMID->SetTextureParameterValue("FoamTex", BakedFoamTexture);
 					WaterInfoMID->SetTextureParameterValue("BakedWaterSimNormalTex", BakedWaterSurfaceNormalTexture);
 					WaterInfoMID->SetVectorParameterValue("BakedWaterSimLocation", SystemPos);
-					WaterMID->SetDoubleVectorParameterValue("BakedWaterSimLocationDouble", SystemPos);
+					WaterInfoMID->SetDoubleVectorParameterValue("BakedWaterSimLocationDouble", SystemPos);
 					WaterInfoMID->SetVectorParameterValue("BakedWaterSimSize", FVector(WorldGridSize.X, WorldGridSize.Y, 1));
 				}
 				else
@@ -1143,10 +1273,30 @@ void UShallowWaterRiverComponent::UpdateRenderState()
 void UShallowWaterRiverComponent::SetWaterMIDParameters(UMaterialInstanceDynamic* WaterMID)
 {
 	if (WaterMID)
-	{
-		WaterMID->SetTextureParameterValue("BakedWaterSimTex", BakedWaterSurfaceTexture);
-		WaterMID->SetTextureParameterValue("FoamTex", BakedFoamTexture);
-		WaterMID->SetTextureParameterValue("BakedWaterSimNormalTex", BakedWaterSurfaceNormalTexture);
+	{		
+		if (bUseVirtualTextures)
+		{
+			WaterMID->SetTextureParameterValue("BakedWaterSimTexVT", BakedWaterSurfaceTexture);
+			WaterMID->SetTextureParameterValue("FoamTexVT", BakedFoamTexture);
+			WaterMID->SetTextureParameterValue("BakedWaterSimNormalTexVT", BakedWaterSurfaceNormalTexture);
+
+			UTexture2D *EmptyPtr = nullptr;
+			WaterMID->SetTextureParameterValue("BakedWaterSimTex", EmptyPtr);
+			WaterMID->SetTextureParameterValue("FoamTex", EmptyPtr);
+			WaterMID->SetTextureParameterValue("BakedWaterSimNormalTex", EmptyPtr);
+		}
+		else
+		{
+			WaterMID->SetTextureParameterValue("BakedWaterSimTex", BakedWaterSurfaceTexture);
+			WaterMID->SetTextureParameterValue("FoamTex", BakedFoamTexture);
+			WaterMID->SetTextureParameterValue("BakedWaterSimNormalTex", BakedWaterSurfaceNormalTexture);
+
+			UTexture2D *EmptyPtr = nullptr;
+			WaterMID->SetTextureParameterValue("BakedWaterSimTexVT", EmptyPtr);
+			WaterMID->SetTextureParameterValue("FoamTexVT", EmptyPtr);
+			WaterMID->SetTextureParameterValue("BakedWaterSimNormalTexVT", EmptyPtr);
+		}
+
 
 		WaterMID->SetVectorParameterValue("BakedWaterSimLocation", SystemPos);
 		WaterMID->SetDoubleVectorParameterValue("BakedWaterSimLocationDouble", SystemPos);
@@ -1154,18 +1304,11 @@ void UShallowWaterRiverComponent::SetWaterMIDParameters(UMaterialInstanceDynamic
 
 		WaterMID->SetTextureParameterValue("NormalDetailTex", NormalDetailRT);
 
-		float Dx;
-		if (WorldGridSize.X > WorldGridSize.Y)
-		{
-			Dx = WorldGridSize.X / ResolutionMaxAxis;
-		}
-		else
-		{
-			Dx = WorldGridSize.Y / ResolutionMaxAxis;
-		}
-		WaterMID->SetScalarParameterValue("BakedWaterSimDx", Dx);
+		WaterMID->SetScalarParameterValue("BakedWaterSimDx", SimDx);
 
 		WaterMID->SetScalarParameterValue("UseBakedSimHack", RenderState == EShallowWaterRenderState::WaterComponentWithBakedSim ? 1 : 0);
+
+		WaterMID->SetVectorParameterValue("SimRes", FVector(SimRes.X, SimRes.Y, 0));
 	}
 	else
 	{
diff --git a/Engine/Plugins/Experimental/WaterAdvanced/Source/WaterAdvanced/Public/ShallowWaterRiverActor.h b/Engine/Plugins/Experimental/WaterAdvanced/Source/WaterAdvanced/Public/ShallowWaterRiverActor.h
--- a/Engine/Plugins/Experimental/WaterAdvanced/Source/WaterAdvanced/Public/ShallowWaterRiverActor.h
+++ b/Engine/Plugins/Experimental/WaterAdvanced/Source/WaterAdvanced/Public/ShallowWaterRiverActor.h
@@ -65,6 +65,9 @@ public:
 	UPROPERTY(EditAnywhere, Category = "Simulation", meta = (DisplayName = "Remove Outside Spline Amount", EditCondition = "bMatchSpline"))
 	float RemoveOutsideSplineAmount = 50.;
 
+	UPROPERTY(EditAnywhere, Category = "Simulation")
+	bool bDisableCPUThrottling = true;
+
 	UPROPERTY(EditAnywhere, Category = "Water", meta = (DisplayName = "Source River Water Bodies"))
 	TArray<TSoftObjectPtr<AWaterBody>> SourceRiverWaterBodies;
 
@@ -98,13 +101,22 @@ public:
 	UPROPERTY(EditAnywhere, Category = "Rendering", meta = (DisplayName = "Smoothing Cutoff"))
 	float SmoothingCutoff = 500.f;
 
-	UPROPERTY(EditAnywhere, Category = "Shallow Water")
+	UPROPERTY()
+	bool bUseVirtualTextures = true;
+
+	UPROPERTY()
+	FVector2D SimRes;
+
+	UPROPERTY()
+	float  SimDx;
+
+	UPROPERTY()
 	TObjectPtr<UTexture2D> BakedWaterSurfaceTexture;
 	
-	UPROPERTY(EditAnywhere, Category = "Shallow Water")
+	UPROPERTY()
 	TObjectPtr<UTexture2D> BakedFoamTexture;
 
-	UPROPERTY(EditAnywhere, Category = "Shallow Water")
+	UPROPERTY()
 	TObjectPtr<UTexture2D> BakedWaterSurfaceNormalTexture;
 
 	//UPROPERTY(EditAnywhere, Category = "Shallow Water")
@@ -154,10 +166,14 @@ public:
 
 	UE_API void Bake();
 
+	UE_API void InitializeVirtualTexture(TObjectPtr<UTexture2D> InTexture);
+
 	UE_API virtual void PostEditChangeProperty(FPropertyChangedEvent& PropertyChangedEvent) override;
 
 	UFUNCTION()
 	UE_API void OnWaterInfoTextureArrayCreated(const UTextureRenderTarget2DArray* InWaterInfoTexture);
+
+	UE_API bool ShouldDisableCPUThrottling() { return true; }
 #endif // WITH_EDITOR
 
 	UFUNCTION(BlueprintCallable, Category = "Shallow Water River")
@@ -184,7 +200,8 @@ protected:
 	TObjectPtr<UBakedShallowWaterSimulationComponent> BakedSim;
 
 	UE_API bool QueryWaterAtSplinePoint(TSoftObjectPtr<AWaterBody> WaterBody, int SplinePoint, FVector& OutPos, FVector& OutTangent, float& OutWidth, float& OutDepth);
-	
+	UE_API void ConvertToVirtualTextures();
+
 private:
 	bool bIsInitialized;	
 	bool bTickInitialize;
@@ -202,10 +219,14 @@ private:
 	UPROPERTY()
 	TObjectPtr<UTextureRenderTarget2D> NormalDetailRT;
 
+	FDelegateHandle ShouldDisableCPUThrottlingDelegateHandle;
+
 	// Find/create the level set renderer singleton actor as required. Return whether the found or created actor.
 	TObjectPtr<UTextureRenderTarget2D> GetSharedFFTOceanPatchNormalRTFromSubsystem(UWorld* World);
 
 	FBoxSphereBounds InitializeCaptureDI(const FName &DIName, TArray<AActor*> RawActorPtrArray);
+	
+	
 };
 
 UCLASS(MinimalAPI, BlueprintType, HideCategories = (Physics, Replication, Input, Collision))
