small optimizations and cleanup
fix ripple sim

#virtualized

diff --git a/Engine/Plugins/Experimental/WaterAdvanced/Source/WaterAdvanced/Private/ShallowWaterSubsystem.cpp b/Engine/Plugins/Experimental/WaterAdvanced/Source/WaterAdvanced/Private/ShallowWaterSubsystem.cpp
--- a/Engine/Plugins/Experimental/WaterAdvanced/Source/WaterAdvanced/Private/ShallowWaterSubsystem.cpp
+++ b/Engine/Plugins/Experimental/WaterAdvanced/Source/WaterAdvanced/Private/ShallowWaterSubsystem.cpp
@@ -595,6 +595,8 @@ void UShallowWaterSubsystem::UpdateGridMovement()
 				{
 					ShallowWaterNiagaraSimulation->SetVariableVec3(FName("BakedWaterSimLocation"), BakedSim->SimulationData.Position);
 					ShallowWaterNiagaraSimulation->SetVariableVec2(FName("BakedWaterSimSize"), BakedSim->SimulationData.Size);
+					ShallowWaterNiagaraSimulation->SetVariableVec2(FName("BakedWaterSimRes"), 
+						FVector2D(BakedSim->SimulationData.NumCells.X, BakedSim->SimulationData.NumCells.Y));
 					ShallowWaterNiagaraSimulation->SetVariableTexture(FName("BakedWaterSimTexture"), BakedSimTex);
 
 					UseBakedSim = true;				
