fix collision dilation for shallow water rivers

diff --git a/Engine/Plugins/Experimental/WaterAdvanced/Source/WaterAdvanced/Private/ShallowWaterRiverActor.cpp b/Engine/Plugins/Experimental/WaterAdvanced/Source/WaterAdvanced/Private/ShallowWaterRiverActor.cpp
--- a/Engine/Plugins/Experimental/WaterAdvanced/Source/WaterAdvanced/Private/ShallowWaterRiverActor.cpp
+++ b/Engine/Plugins/Experimental/WaterAdvanced/Source/WaterAdvanced/Private/ShallowWaterRiverActor.cpp
@@ -968,43 +968,65 @@ void UShallowWaterRiverComponent::Bake()
 
 			// make sure the collision convex hull vertices are clamped to the min/max water height
 			for (FKConvexElem& ConvexElem : CurrBodySetup->AggGeom.ConvexElems)
-			{			
-				// see if the current point is inside the convex projected to the xy plane
-				FBox CurrBox = ConvexElem.CalcAABB(CurrMeshTransform, FVector(1, 1, 1));				
-
+			{
 				TArray<FVector>& VertexData = ConvexElem.VertexData;
-				
-				// 
-				// We know based on the way each convex elem is created which indices correspond with
-				// top and bottom of the box.  From SplineMeshComponent.cpp:1075
-				// ConvexElem.VertexData.Add(ElementTM.TransformPosition(Radii * FVector(-1, -1, -1)));
-				// ConvexElem.VertexData.Add(ElementTM.TransformPosition(Radii * FVector(-1, -1, 1)));
-				// ConvexElem.VertexData.Add(ElementTM.TransformPosition(Radii * FVector(-1, 1, -1)));
-				// ConvexElem.VertexData.Add(ElementTM.TransformPosition(Radii * FVector(-1, 1, 1)));
-				// ConvexElem.VertexData.Add(ElementTM.TransformPosition(Radii * FVector(1, -1, -1)));
-				// ConvexElem.VertexData.Add(ElementTM.TransformPosition(Radii * FVector(1, -1, 1)));
-				// ConvexElem.VertexData.Add(ElementTM.TransformPosition(Radii * FVector(1, 1, -1)));
-				// ConvexElem.VertexData.Add(ElementTM.TransformPosition(Radii * FVector(1, 1, 1)));
-				
+
 				if (const float* WorldMaxZForConvex = ConvexToMaxHeight.Find(&ConvexElem))
 				{
 					// for each vertex in the convex hull, set the Z to the maximum baked water sim Z height for the convex
 					int32 Idx = 0;
 					for (FVector& Vertex : VertexData)
 					{
-						FVector VWorld = CurrMeshTransform.TransformPosition(Vertex);					
-					
-						// Only set on top vertices of the convex hull with +Z to push it up
-						if (Idx == 1 || Idx == 3 || Idx == 5 || Idx == 7)
+						// only top vertices
+						if (Idx >= 4)
 						{
-							VWorld.Z = *WorldMaxZForConvex;
+							FVector VWorld = CurrMeshTransform.TransformPosition(Vertex);
+
+							VWorld.Z = FMath::Max(VWorld.Z, *WorldMaxZForConvex);
+
+							const FVector VLocal = CurrMeshTransform.InverseTransformPosition(VWorld);
+							Vertex.X = VLocal.X;
+							Vertex.Y = VLocal.Y;
+							Vertex.Z = VLocal.Z;
 						}
 
-						const FVector VLocal = CurrMeshTransform.InverseTransformPosition(VWorld);
-						Vertex.X = VLocal.X;
-						Vertex.Y = VLocal.Y;
-						Vertex.Z = VLocal.Z;
 
+						#if ENABLE_DRAW_DEBUG
+						if (bShallowWaterRiverDebugVisualize)
+						{		
+							FVector VWorld = CurrMeshTransform.TransformPosition(Vertex);
+
+							switch (Idx)
+							{
+								case 0:
+								DrawDebugSphere(GetWorld() , VWorld, 10., 2, FColor::Red, true);		
+								break;
+								case 1:
+								DrawDebugSphere(GetWorld() , VWorld, 10., 3, FColor::Green, true);		
+								break;
+								case 2:
+								DrawDebugSphere(GetWorld() , VWorld, 10., 4, FColor::Blue, true);		
+								break;
+								case 3:
+								DrawDebugSphere(GetWorld() , VWorld, 10., 5, FColor::Black, true);		
+								break;
+								case 4:
+								DrawDebugSphere(GetWorld() , VWorld, 10., 6, FColor::White, true);	//	
+								break;
+								case 5:
+								DrawDebugSphere(GetWorld() , VWorld, 10., 7, FColor::Magenta, true); //		
+								break;
+								case 6:
+								DrawDebugSphere(GetWorld() , VWorld, 10., 8, FColor::Orange, true);	//	
+								break;
+								case 7:
+								DrawDebugSphere(GetWorld() , VWorld, 10., 9, FColor::Purple, true);	//	
+								break;								
+							}
+							
+						}
+						#endif
+												
 						Idx++;
 					}
 				}
