improve comment

diff --git a/Engine/Plugins/Experimental/WaterAdvanced/Source/WaterAdvanced/Private/ShallowWaterRiverActor.cpp b/Engine/Plugins/Experimental/WaterAdvanced/Source/WaterAdvanced/Private/ShallowWaterRiverActor.cpp
--- a/Engine/Plugins/Experimental/WaterAdvanced/Source/WaterAdvanced/Private/ShallowWaterRiverActor.cpp
+++ b/Engine/Plugins/Experimental/WaterAdvanced/Source/WaterAdvanced/Private/ShallowWaterRiverActor.cpp
@@ -977,7 +977,7 @@ void UShallowWaterRiverComponent::Bake()
 					int32 Idx = 0;
 					for (FVector& Vertex : VertexData)
 					{
-						// only top vertices
+						// only top vertices are 4,5,6,7
 						if (Idx >= 4)
 						{
 							FVector VWorld = CurrMeshTransform.TransformPosition(Vertex);
