option to recursively add attached actors for collisions (ie: for megamesh)

diff --git a/Engine/Plugins/Experimental/WaterAdvanced/Source/WaterAdvanced/Private/ShallowWaterRiverActor.cpp b/Engine/Plugins/Experimental/WaterAdvanced/Source/WaterAdvanced/Private/ShallowWaterRiverActor.cpp
--- a/Engine/Plugins/Experimental/WaterAdvanced/Source/WaterAdvanced/Private/ShallowWaterRiverActor.cpp
+++ b/Engine/Plugins/Experimental/WaterAdvanced/Source/WaterAdvanced/Private/ShallowWaterRiverActor.cpp
@@ -729,6 +729,15 @@ void UShallowWaterRiverComponent::AddActorsToRawArray(const TArray<TSoftObjectPt
 					return true;
 				});
 		}
+		else if (bRecursivelyAddAttachedActors)
+		{
+			TArray<AActor*> Meshes;
+			CurrActorRawPtr->GetAttachedActors(Meshes);
+			for (AActor* AttachedMesh : Meshes)
+			{
+				BottomContourActorsRawPtr.Add(AttachedMesh);
+			}
+		}
 		else
 		{
 			BottomContourActorsRawPtr.Add(CurrActorRawPtr);
diff --git a/Engine/Plugins/Experimental/WaterAdvanced/Source/WaterAdvanced/Public/ShallowWaterRiverActor.h b/Engine/Plugins/Experimental/WaterAdvanced/Source/WaterAdvanced/Public/ShallowWaterRiverActor.h
--- a/Engine/Plugins/Experimental/WaterAdvanced/Source/WaterAdvanced/Public/ShallowWaterRiverActor.h
+++ b/Engine/Plugins/Experimental/WaterAdvanced/Source/WaterAdvanced/Public/ShallowWaterRiverActor.h
@@ -125,6 +125,9 @@ public:
 	UPROPERTY(EditAnywhere, Category = "Collisions")
 	bool bUseCapture = true;
 
+	UPROPERTY(EditAnywhere, Category = "Collisions")
+	bool bRecursivelyAddAttachedActors = false;
+
 	UPROPERTY(EditAnywhere, Category = "Collisions")
 	TArray<TSoftObjectPtr<AActor>> BottomContourLandscapeActors;
 
