* Marked Water Advanced Experimental
* Assign Water-looking icons to Water Advanced and Buoyancy plugins

#rnx

diff --git a/Engine/Plugins/Experimental/WaterAdvanced/WaterAdvanced.uplugin b/Engine/Plugins/Experimental/WaterAdvanced/WaterAdvanced.uplugin
--- a/Engine/Plugins/Experimental/WaterAdvanced/WaterAdvanced.uplugin
+++ b/Engine/Plugins/Experimental/WaterAdvanced/WaterAdvanced.uplugin
@@ -1,8 +1,8 @@
 {
 	"FileVersion": 3,
 	"Version": 1,
-	"VersionName": "1.0",
-	"FriendlyName": "WaterAdvanced",
+	"VersionName": "0.1",
+	"FriendlyName": "Water Advanced",
 	"Description": "Collection of easy to use water simulation systems built on the Niagara Fluids and Water plugins",
 	"Category": "Water",
 	"CreatedBy": "Epic Games, Inc.",
@@ -12,8 +12,8 @@
 	"SupportURL": "",
 	"EnabledByDefault": false,
 	"CanContainContent": true,
-	"IsBetaVersion": true,
-	"IsExperimentalVersion": false,
+	"IsBetaVersion": false,
+	"IsExperimentalVersion": true,
 	"Installed": false,
 	"Modules": [
 		{
