fix material compile warnings

#virtualized

diff --git a/Engine/Plugins/Experimental/WaterAdvanced/Shaders/WaterAdvanced.ush b/Engine/Plugins/Experimental/WaterAdvanced/Shaders/WaterAdvanced.ush
--- a/Engine/Plugins/Experimental/WaterAdvanced/Shaders/WaterAdvanced.ush
+++ b/Engine/Plugins/Experimental/WaterAdvanced/Shaders/WaterAdvanced.ush
@@ -74,7 +74,7 @@ float2 CalculateUV(float2 vAxis, float2 position, float Time, float Speed)
 	return uv;
 }
 
-float2 ComputeMipLevel(float2 uv, int2 sz)
+float ComputeMipLevel(float2 uv, int2 sz)
 {	
 	const float2 dxval = ddx(uv) * sz;
 	const float2 dyval = ddy(uv) * sz;
