- Fix warning about ocean patch system not existing in some cooked builds

#rnx

diff --git a/Engine/Plugins/Experimental/WaterAdvanced/Source/WaterAdvanced/Private/FFTOceanPatchSubsystem.cpp b/Engine/Plugins/Experimental/WaterAdvanced/Source/WaterAdvanced/Private/FFTOceanPatchSubsystem.cpp
--- a/Engine/Plugins/Experimental/WaterAdvanced/Source/WaterAdvanced/Private/FFTOceanPatchSubsystem.cpp
+++ b/Engine/Plugins/Experimental/WaterAdvanced/Source/WaterAdvanced/Private/FFTOceanPatchSubsystem.cpp
@@ -2,6 +2,7 @@
 
 #include "FFTOceanPatchSubsystem.h"
 #include "ShallowWaterSettings.h"
+#include "AssetRegistry/AssetData.h"
 #include "Engine/World.h"
 #include "Engine/AssetManager.h"
 #include "NiagaraComponent.h"
@@ -47,15 +48,36 @@ void UFFTOceanPatchSubsystem::PostInitialize()
 		return;
 	}
 
-	TArray< FSoftObjectPath> ObjectsToLoad;
+	TArray<FSoftObjectPath> ObjectsToLoad;
 	ObjectsToLoad.Add(Settings->DefaultOceanPatchNiagaraSystem.ToSoftObjectPath());
 
-	UAssetManager::GetStreamableManager().RequestAsyncLoad(ObjectsToLoad,
-	FStreamableDelegate::CreateWeakLambda(this, [this]()
+	// To avoid warnings in cooked builds do a check with the asset registry that we cooked the assets
+	if (FPlatformProperties::RequiresCookedData())
+	{
+		IAssetRegistry& AssetRegistry = FModuleManager::LoadModuleChecked<FAssetRegistryModule>("AssetRegistry").Get();
+
+		for (auto It=ObjectsToLoad.CreateIterator(); It; ++It)
 		{
-			// continue with initialization		
-		})
-	);	
+			FAssetData AssetData;
+			if (AssetRegistry.TryGetAssetByObjectPath(*It, AssetData) != UE::AssetRegistry::EExists::Exists)
+			{
+				It.RemoveCurrentSwap();
+			}
+		}
+	}
+
+	if (ObjectsToLoad.Num() > 0)
+	{
+		UAssetManager::GetStreamableManager().RequestAsyncLoad(
+			ObjectsToLoad,
+			FStreamableDelegate::CreateWeakLambda(this,
+				[this]()
+				{
+					// continue with initialization		
+				}
+			)
+		);	
+	}
 
 	FFTOceanSystem = nullptr;
 	OceanNormalRT = nullptr;
