fix bouding box creation ordering

#Synced-CL 49020631

diff --git a/Engine/Plugins/Experimental/WaterAdvanced/Source/WaterAdvanced/Private/ShallowWaterRiverActor.cpp b/Engine/Plugins/Experimental/WaterAdvanced/Source/WaterAdvanced/Private/ShallowWaterRiverActor.cpp
--- a/Engine/Plugins/Experimental/WaterAdvanced/Source/WaterAdvanced/Private/ShallowWaterRiverActor.cpp
+++ b/Engine/Plugins/Experimental/WaterAdvanced/Source/WaterAdvanced/Private/ShallowWaterRiverActor.cpp
@@ -408,7 +408,14 @@ void UShallowWaterRiverComponent::Rebuild()
 	}
 	
 	SystemPos = CombinedBounds.Origin - FVector(0, 0, CombinedBounds.BoxExtent.Z);
-	
+	WorldGridSize = 2.0f * FVector2D(CombinedBounds.BoxExtent.X, CombinedBounds.BoxExtent.Y);
+
+	if (WorldGridSize.X < SMALL_NUMBER || WorldGridSize.Y < SMALL_NUMBER)
+	{
+		UE_LOG(LogShallowWater, Warning, TEXT("UShallowWaterRiverComponent::Rebuild() - Simulation grid has (0,0) size."));
+		return;
+	}
+
 	RiverSimSystem = NewObject<UNiagaraComponent>(this, NAME_None, RF_Public);
 	RiverSimSystem->bUseAttachParentBound = false;
 	RiverSimSystem->SetWorldLocation(SystemPos);
@@ -522,14 +529,6 @@ void UShallowWaterRiverComponent::Rebuild()
 	}
 
 	RiverSimSystem->Activate();
-	
-	WorldGridSize = 2.0f * FVector2D(CombinedBounds.BoxExtent.X, CombinedBounds.BoxExtent.Y);
-
-	if (WorldGridSize.Length() < 1e-8)
-	{
-		UE_LOG(LogShallowWater, Warning, TEXT("UShallowWaterRiverComponent::Bake() - Simulation grid has (0,0) size."));
-		return;
-	}
 
 	RiverSimSystem->SetVariableVec2(FName("WorldGridSize"), WorldGridSize);
 	RiverSimSystem->SetVariableInt(FName("ResolutionMaxAxis"), ResolutionMaxAxis);
@@ -812,7 +811,7 @@ void UShallowWaterRiverComponent::Bake()
 		{
 			if (CurrWaterBody == nullptr)
 			{			
-				UE_LOG(LogShallowWater, Warning, TEXT("UShallowWaterRiverComponent::Rebuild() - Cannot bake with a null water body.  Please make sure all water bodies are loaded and that all entires on the actor are valid"));
+				UE_LOG(LogShallowWater, Warning, TEXT("UShallowWaterRiverComponent::Bake() - Cannot bake with a null water body.  Please make sure all water bodies are loaded and that all entries on the actor are valid"));
 				return;
 			}
 		}
