underwater post process material for shallow water
fix power of 2 sampling issues with VT (stretch instead of pad)

#virtualized
#Synced-CL 49084831

diff --git a/Engine/Plugins/Experimental/WaterAdvanced/Source/WaterAdvanced/Private/ShallowWaterRiverActor.cpp b/Engine/Plugins/Experimental/WaterAdvanced/Source/WaterAdvanced/Private/ShallowWaterRiverActor.cpp
--- a/Engine/Plugins/Experimental/WaterAdvanced/Source/WaterAdvanced/Private/ShallowWaterRiverActor.cpp
+++ b/Engine/Plugins/Experimental/WaterAdvanced/Source/WaterAdvanced/Private/ShallowWaterRiverActor.cpp
@@ -140,7 +140,7 @@ void UShallowWaterRiverComponent::ConvertToVirtualTextures()
 
 	bUseVirtualTextures = true;
 
-	if (BakedWaterSurfaceTexture != NULL && !BakedWaterSurfaceTexture->VirtualTextureStreaming)
+	if (BakedWaterSurfaceTexture != NULL && !BakedWaterSurfaceTexture->VirtualTextureStreaming && bUseVirtualTextures)
 	{	
 		SimRes = FVector2D(BakedWaterSurfaceTexture->Source.GetSizeX(), BakedWaterSurfaceTexture->Source.GetSizeY());
 		RiverSimSystem->SetVariableVec2(FName("SimRes"), SimRes);
@@ -150,22 +150,46 @@ void UShallowWaterRiverComponent::ConvertToVirtualTextures()
 
 		HasChanged = true;
 	}
+	else if (BakedWaterSurfaceTexture != NULL && BakedWaterSurfaceTexture->VirtualTextureStreaming && 
+		BakedWaterSurfaceTexture->PowerOfTwoMode != ETexturePowerOfTwoSetting::StretchToPowerOfTwo && bUseVirtualTextures)
+	{
+		InitializeVirtualTexture(BakedWaterSurfaceTexture);
+		UE_LOG(LogShallowWater, Warning, TEXT("Baked water surface texture wrong power of two mode - converting.  Recommended resave."));
+
+		HasChanged = true;
+	}
 	
-	if (BakedWaterSurfaceNormalTexture != NULL && !BakedWaterSurfaceNormalTexture->VirtualTextureStreaming)
+	if (BakedWaterSurfaceNormalTexture != NULL && !BakedWaterSurfaceNormalTexture->VirtualTextureStreaming && bUseVirtualTextures)
 	{
 		InitializeVirtualTexture(BakedWaterSurfaceNormalTexture);
 		UE_LOG(LogShallowWater, Warning, TEXT("Baked water surface normal texture was not virtual- converting.  Recommended resave."));
 
 		HasChanged = true;
 	}
+	else if (BakedWaterSurfaceNormalTexture != NULL && BakedWaterSurfaceNormalTexture->VirtualTextureStreaming && 
+		BakedWaterSurfaceNormalTexture->PowerOfTwoMode != ETexturePowerOfTwoSetting::StretchToPowerOfTwo && bUseVirtualTextures)
+	{
+		InitializeVirtualTexture(BakedWaterSurfaceNormalTexture);
+		UE_LOG(LogShallowWater, Warning, TEXT("Baked water surface normal texture wrong power of two mode - converting.  Recommended resave."));
 
-	if (BakedFoamTexture != NULL && !BakedFoamTexture->VirtualTextureStreaming)
+		HasChanged = true;
+	}
+
+	if (BakedFoamTexture != NULL && !BakedFoamTexture->VirtualTextureStreaming && bUseVirtualTextures)
 	{
 		InitializeVirtualTexture(BakedFoamTexture);
 		UE_LOG(LogShallowWater, Warning, TEXT("Baked foam texture was not virtual- converting.  Recommended resave."));
 
 		HasChanged = true;
 	}
+	else if (BakedFoamTexture != NULL && BakedFoamTexture->VirtualTextureStreaming && 
+		BakedFoamTexture->PowerOfTwoMode != ETexturePowerOfTwoSetting::StretchToPowerOfTwo && bUseVirtualTextures)
+	{
+		InitializeVirtualTexture(BakedFoamTexture);
+		UE_LOG(LogShallowWater, Warning, TEXT("Baked water surface foam texture wrong power of two mode - converting.  Recommended resave."));
+
+		HasChanged = true;
+	}
 
 	if (HasChanged)
 	{
@@ -409,7 +433,7 @@ void UShallowWaterRiverComponent::Rebuild()
 	
 	SystemPos = CombinedBounds.Origin - FVector(0, 0, CombinedBounds.BoxExtent.Z);
 	WorldGridSize = 2.0f * FVector2D(CombinedBounds.BoxExtent.X, CombinedBounds.BoxExtent.Y);
-
+	
 	if (WorldGridSize.X < SMALL_NUMBER || WorldGridSize.Y < SMALL_NUMBER)
 	{
 		UE_LOG(LogShallowWater, Warning, TEXT("UShallowWaterRiverComponent::Rebuild() - Simulation grid has (0,0) size."));
@@ -529,16 +553,30 @@ void UShallowWaterRiverComponent::Rebuild()
 	}
 
 	RiverSimSystem->Activate();
-
+	
 	RiverSimSystem->SetVariableVec2(FName("WorldGridSize"), WorldGridSize);
 	RiverSimSystem->SetVariableInt(FName("ResolutionMaxAxis"), ResolutionMaxAxis);
 
 	// #todo(dmp): would be better to initialize the user var inside of niagara rather than recomputing res here
-	SimRes = FVector2D(ResolutionMaxAxis, ResolutionMaxAxis * WorldGridSize.Y / WorldGridSize.X);
+	double CellSize = FMath::Max(WorldGridSize.X, WorldGridSize.Y) / ResolutionMaxAxis;
+
+	SimRes = FVector2D(ResolutionMaxAxis, FMath::FloorToInt(ResolutionMaxAxis * WorldGridSize.Y / WorldGridSize.X));
+	
 	if (WorldGridSize.Y > WorldGridSize.X)
 	{
-		SimRes = FVector2D(ResolutionMaxAxis * WorldGridSize.X / WorldGridSize.Y, ResolutionMaxAxis);
+		SimRes = FVector2D(FMath::FloorToInt(ResolutionMaxAxis * WorldGridSize.X / WorldGridSize.Y), ResolutionMaxAxis);
 	}
+
+	if (WorldGridSize.X > WorldGridSize.Y && FMath::Abs(CellSize * SimRes.Y - WorldGridSize.Y) > SMALL_NUMBER)
+	{
+		SimRes.Y++;
+	}
+
+	if (WorldGridSize.X < WorldGridSize.Y && FMath::Abs(CellSize * SimRes.X - WorldGridSize.X) > SMALL_NUMBER)
+	{
+		SimRes.X++;
+	}
+
 	RiverSimSystem->SetVariableVec2(FName("SimRes"), SimRes);
 
 	// pad out source's box height a so it intersects the sim plane.  This value doesn't matter much so we hardcode it
@@ -830,7 +868,7 @@ void UShallowWaterRiverComponent::Bake()
 	}		
 
 	BakedWaterSurfaceTexture = BakedWaterSurfaceRT->ConstructTexture2D(this, "BakedRiverTexture", TextureObjectFlags);
-
+	
 	if (bUseVirtualTextures)
 	{
 		InitializeVirtualTexture(BakedWaterSurfaceTexture);
@@ -888,12 +926,15 @@ void UShallowWaterRiverComponent::Bake()
 	{ 
 		for (TSoftObjectPtr<AWaterBody > CurrWaterBody : BakedSim->WaterBodies)
 		{
-			TObjectPtr<UWaterBodyComponent> CurrWaterBodyComponent = CurrWaterBody->GetWaterBodyComponent();
-
-			if (CurrWaterBodyComponent != nullptr)
+			if (CurrWaterBody != nullptr)
 			{
-				CurrWaterBodyComponent->SetBakedShallowWaterSimulation(nullptr);
-				CurrWaterBodyComponent->PostEditChange();
+				TObjectPtr<UWaterBodyComponent> CurrWaterBodyComponent = CurrWaterBody->GetWaterBodyComponent();
+
+				if (CurrWaterBodyComponent != nullptr)
+				{
+					CurrWaterBodyComponent->SetBakedShallowWaterSimulation(nullptr);
+					CurrWaterBodyComponent->PostEditChange();
+				}
 			}
 		}
 	}
@@ -1051,7 +1092,7 @@ void UShallowWaterRiverComponent::InitializeVirtualTexture(TObjectPtr<UTexture2D
 {	
 	InTexture->Modify();
 	InTexture->MipGenSettings = TextureMipGenSettings::TMGS_SimpleAverage;
-	InTexture->PowerOfTwoMode = ETexturePowerOfTwoSetting::PadToPowerOfTwo;
+	InTexture->PowerOfTwoMode = ETexturePowerOfTwoSetting::StretchToPowerOfTwo;
 	InTexture->VirtualTextureStreaming = true;
 
 	InTexture->UpdateResource();	
@@ -1175,6 +1216,11 @@ void UShallowWaterRiverComponent::UpdateRenderState()
 		BakedSimMaterial = LoadObject<UMaterialInstance>(nullptr, TEXT("/WaterAdvanced/Niagara/Materials/SW_Water_Material_River.SW_Water_Material_River"));
 	}
 	
+	if (BakedSimUnderWaterMaterial == nullptr)
+	{
+		BakedSimUnderWaterMaterial = LoadObject<UMaterialInstance>(nullptr, TEXT("/WaterAdvanced/Materials/M_UnderWater_PostProcess_Volume_SW.M_UnderWater_PostProcess_Volume_SW"));
+	}
+
 	if (BakedSimRiverToLakeTransitionMaterial == nullptr)
 	{
 		BakedSimRiverToLakeTransitionMaterial = LoadObject<UMaterialInstance>(nullptr, TEXT("/WaterAdvanced/Niagara/Materials/SW_Water_Material_River_To_Lake_Transition.SW_Water_Material_River_To_Lake_Transition"));
@@ -1200,6 +1246,11 @@ void UShallowWaterRiverComponent::UpdateRenderState()
 		SplineRiverToOceanTransitionMaterial = LoadObject<UMaterialInstance>(nullptr, TEXT("/WaterAdvanced/Niagara/Materials/SW_Water_Material_River_To_Ocean_Transition_Spline.SW_Water_Material_River_To_Ocean_Transition_Spline"));
 	}
 
+	if (SplineRiverUnderWaterMaterial == nullptr)
+	{
+		SplineRiverUnderWaterMaterial = LoadObject<UMaterialInstance>(nullptr, TEXT("/WaterAdvanced/Materials/M_UnderWater_PostProcess_Volume_Spline.M_UnderWater_PostProcess_Volume_Spline"));
+	}
+
 	bool bReadBakedSim = RenderState == EShallowWaterRenderState::BakedSim || RenderState == EShallowWaterRenderState::WaterComponentWithBakedSim || RenderState == EShallowWaterRenderState::WaterComponent;
 	bool RenderWaterBody = RenderState == EShallowWaterRenderState::WaterComponent || RenderState == EShallowWaterRenderState::WaterComponentWithBakedSim;
 	bool RenderSecondary = 
@@ -1246,6 +1297,10 @@ void UShallowWaterRiverComponent::UpdateRenderState()
 				UMaterialInstanceDynamic* WaterMID = CurrWaterBodyComponent->GetWaterMaterialInstance();			
 				SetWaterMIDParameters(WaterMID);
 
+				CurrWaterBodyComponent->SetUnderwaterPostProcessMaterial(BakedSimUnderWaterMaterial);
+				UMaterialInstanceDynamic* UnderWaterMID = CurrWaterBodyComponent->GetUnderwaterPostProcessMaterialInstance();
+				SetWaterMIDParameters(UnderWaterMID);
+
 				CurrWaterBodyComponent->SetLakeTransitionMaterial(BakedSimRiverToLakeTransitionMaterial);
 				UMaterialInstanceDynamic* WaterLakeTransitionMID = CurrWaterBodyComponent->GetRiverToLakeTransitionMaterialInstance();
 				SetWaterMIDParameters(WaterLakeTransitionMID);
@@ -1275,6 +1330,7 @@ void UShallowWaterRiverComponent::UpdateRenderState()
 				CurrWaterBodyComponent->SetWaterMaterial(SplineRiverMaterial);
 				CurrWaterBodyComponent->SetLakeTransitionMaterial(SplineRiverToLakeTransitionMaterial);
 				CurrWaterBodyComponent->SetOceanTransitionMaterial(SplineRiverToOceanTransitionMaterial);
+				CurrWaterBodyComponent->SetUnderwaterPostProcessMaterial(SplineRiverUnderWaterMaterial);
 			}
 
 			CurrWaterBodyComponent->SetUseBakedSimulationForQueriesAndPhysics(
diff --git a/Engine/Plugins/Experimental/WaterAdvanced/Source/WaterAdvanced/Public/ShallowWaterRiverActor.h b/Engine/Plugins/Experimental/WaterAdvanced/Source/WaterAdvanced/Public/ShallowWaterRiverActor.h
--- a/Engine/Plugins/Experimental/WaterAdvanced/Source/WaterAdvanced/Public/ShallowWaterRiverActor.h
+++ b/Engine/Plugins/Experimental/WaterAdvanced/Source/WaterAdvanced/Public/ShallowWaterRiverActor.h
@@ -80,6 +80,9 @@ public:
 	UPROPERTY(EditAnywhere, Category = "Rendering", meta = (DisplayName = "Baked Sim Material"))
 	TObjectPtr <class UMaterialInstance> BakedSimMaterial;
 
+	UPROPERTY(EditAnywhere, Category = "Rendering", meta = (DisplayName = "Baked UnderWater Sim Material"))
+	TObjectPtr <class UMaterialInstance> BakedSimUnderWaterMaterial;
+
 	UPROPERTY(EditAnywhere, Category = "Rendering", meta = (DisplayName = "Baked Sim River To Lake Transition Material"))
 	TObjectPtr <class UMaterialInstance> BakedSimRiverToLakeTransitionMaterial;
 
@@ -95,6 +98,9 @@ public:
 	UPROPERTY(EditAnywhere, Category = "Rendering", meta = (DisplayName = "Spline River To Ocean Transition Material"))
 	TObjectPtr <class UMaterialInstance> SplineRiverToOceanTransitionMaterial;
 
+	UPROPERTY(EditAnywhere, Category = "Rendering", meta = (DisplayName = "Spline River UnderWater Material"))
+	TObjectPtr <class UMaterialInstance> SplineRiverUnderWaterMaterial;
+
 	UPROPERTY(EditAnywhere, Category = "Rendering", meta = (DisplayName = "Smoothing Width"))
 	int SmoothingWidth = 5;
 
@@ -109,7 +115,7 @@ public:
 
 	UPROPERTY()
 	float  SimDx;
-
+	
 	UPROPERTY()
 	TObjectPtr<UTexture2D> BakedWaterSurfaceTexture;
 	
@@ -169,7 +175,7 @@ public:
 
 	UE_API void Bake();
 
-	UE_API void InitializeVirtualTexture(TObjectPtr<UTexture2D> InTexture);
+	UE_API void InitializeVirtualTexture(TObjectPtr<UTexture2D> InTexture);	
 
 	UE_API virtual void PostEditChangeProperty(FPropertyChangedEvent& PropertyChangedEvent) override;
 
