fix for packaged builds

#virtualized
#Synced-CL 49343818

diff --git a/Engine/Plugins/Experimental/WaterAdvanced/Source/WaterAdvanced/Private/ShallowWaterRiverActor.cpp b/Engine/Plugins/Experimental/WaterAdvanced/Source/WaterAdvanced/Private/ShallowWaterRiverActor.cpp
--- a/Engine/Plugins/Experimental/WaterAdvanced/Source/WaterAdvanced/Private/ShallowWaterRiverActor.cpp
+++ b/Engine/Plugins/Experimental/WaterAdvanced/Source/WaterAdvanced/Private/ShallowWaterRiverActor.cpp
@@ -321,6 +321,11 @@ void UShallowWaterRiverComponent::Rebuild()
 		NiagaraRiverSimulation = LoadObject<UNiagaraSystem>(nullptr, TEXT("/WaterAdvanced/Niagara/Systems/Grid2D_SW_River.Grid2D_SW_River"));
 	}
 
+	if (OceanPatchSystem == nullptr)
+	{
+		OceanPatchSystem = LoadObject<UNiagaraSystem>(nullptr, TEXT("/WaterAdvanced/Niagara/Systems/Grid2D_OceanPatch.Grid2D_OceanPatch"));
+	}
+
 	if (RiverSimSystem != nullptr)
 	{
 		RiverSimSystem->SetActive(false);
diff --git a/Engine/Plugins/Experimental/WaterAdvanced/Source/WaterAdvanced/Public/ShallowWaterRiverActor.h b/Engine/Plugins/Experimental/WaterAdvanced/Source/WaterAdvanced/Public/ShallowWaterRiverActor.h
--- a/Engine/Plugins/Experimental/WaterAdvanced/Source/WaterAdvanced/Public/ShallowWaterRiverActor.h
+++ b/Engine/Plugins/Experimental/WaterAdvanced/Source/WaterAdvanced/Public/ShallowWaterRiverActor.h
@@ -228,6 +228,9 @@ private:
 	UPROPERTY()
 	TObjectPtr<UTextureRenderTarget2D> NormalDetailRT;
 
+	UPROPERTY()
+	TObjectPtr<class UNiagaraSystem> OceanPatchSystem;
+
 	FDelegateHandle ShouldDisableCPUThrottlingDelegateHandle;
 
 	// Find/create the level set renderer singleton actor as required. Return whether the found or created actor.
