avoid crash with null ndc

#Synced-CL 50665815

diff --git a/Engine/Plugins/Experimental/WaterAdvanced/Source/WaterAdvanced/Private/ShallowWaterSubsystem.cpp b/Engine/Plugins/Experimental/WaterAdvanced/Source/WaterAdvanced/Private/ShallowWaterSubsystem.cpp
--- a/Engine/Plugins/Experimental/WaterAdvanced/Source/WaterAdvanced/Private/ShallowWaterSubsystem.cpp
+++ b/Engine/Plugins/Experimental/WaterAdvanced/Source/WaterAdvanced/Private/ShallowWaterSubsystem.cpp
@@ -683,6 +683,11 @@ void UShallowWaterSubsystem::WriteImpactToNDC(FVector ImpactPosition, FVector Im
 {
 	FNiagaraDataChannelSearchParameters SearchParams(ImpactPosition);
 	const UNiagaraDataChannelAsset* NDC = Settings->DefaultShallowWaterCollisionNDC.Get();
+	if (NDC == nullptr)
+	{
+		return;
+	}
+	
 	if (UNiagaraDataChannelWriter* DCWriter = UNiagaraDataChannelLibrary::WriteToNiagaraDataChannel(ShallowWaterNiagaraSimulation, NDC, SearchParams, 1, false, true, true, TEXT("ShallowWaterWriteImpact")))
 	{
 		int32 Index = 0;
