move water sim data to uobject and pass pointer around
support for velocity buffer

#Synced-CL 51253254

#Synced-CL 51253254

diff --git a/Engine/Config/BaseEngine.ini b/Engine/Config/BaseEngine.ini
--- a/Engine/Config/BaseEngine.ini
+++ b/Engine/Config/BaseEngine.ini
@@ -1397,6 +1397,10 @@ ClassName=/Script/OnlineSubsystemUtils.OnlineEngineInterfaceImpl
 
 +ClassRedirects=(OldName="/Script/Engine.WorldPartitionDestructibleHLODMeshComponent",NewName="/Script/Engine.WorldPartitionDestructibleHLODComponent")
 
+; deprecate old method for storing baked water simulation data
++StructRedirects=(OldName="/Script/Water.ShallowWaterSimulationGrid",NewName="/Script/Water.ShallowWaterSimulationGrid_DEPRECATED")
++PropertyRedirects=(OldName="/Script/Water.BakedShallowWaterSimulationComponent.SimulationData",NewName="/Script/Water.BakedShallowWaterSimulationComponent.SimulationData_DEPRECATED")
+
 ; End of releases redirects
 
 ; Fixing the output of tags for blueprints to be reset.
diff --git a/Engine/Plugins/Experimental/Buoyancy/Source/Runtime/Private/BuoyancySubsystem.cpp b/Engine/Plugins/Experimental/Buoyancy/Source/Runtime/Private/BuoyancySubsystem.cpp
--- a/Engine/Plugins/Experimental/Buoyancy/Source/Runtime/Private/BuoyancySubsystem.cpp
+++ b/Engine/Plugins/Experimental/Buoyancy/Source/Runtime/Private/BuoyancySubsystem.cpp
@@ -553,15 +553,14 @@ void UBuoyancySubsystem::UpdateSplineData()
 				// Copy out water spline data into a shared ptr, to be associated with all
 				// child particles and marshaled to PT.
 				const Chaos::FRigidTransform3 WaterTransform = WaterBodyComponent->GetComponentTransform();
-				const TOptional<FInterpCurveFloat> EmptyOptionalFloat;
-				const TOptional<FShallowWaterSimulationGrid> EmptyOptionalSimGrid;
+				const TOptional<FInterpCurveFloat> EmptyOptionalFloat;				
 				TSharedPtr<FBuoyancyWaterSplineData> WaterSplineData = MakeShared<FBuoyancyWaterSplineData>(
 					WaterTransform,
 					SplineComponent->GetSplinePointsPosition(),
 					WaterBodyComponent->GetWaterBodyType(),
 					SplineMetadata ? SplineMetadata->RiverWidth : EmptyOptionalFloat,
 					SplineMetadata ? SplineMetadata->WaterVelocityScalar : EmptyOptionalFloat,
-					bUseShallowWaterSimulation && WaterBodyComponent->UseBakedSimulationForQueriesAndPhysics() ? BakedSim->SimulationData : EmptyOptionalSimGrid
+					bUseShallowWaterSimulation && WaterBodyComponent->UseBakedSimulationForQueriesAndPhysics() ? BakedSim->BakedSimulationData : nullptr
 				);
 
 				// Go over each physics object in each primitive component which was generated
@@ -1125,9 +1124,8 @@ void FBuoyancySubsystemSimCallback::TrackInteraction(
 
 		FVector ClosestWaterPointToParticle = ParticlePos;
 		ClosestWaterPointToParticle.Z = WaterHeight;
-
-		FShallowWaterSimulationGrid DefaultShallowWaterSimulationGrid;
-		const FShallowWaterSimulationGrid& ShallowWaterSimData = WaterSpline.ShallowWaterSimData.Get(DefaultShallowWaterSimulationGrid);
+		
+		const TWeakObjectPtr<UShallowWaterSimulationData> ShallowWaterSimData = WaterSpline.ShallowWaterSimData;
 
 		WaterSampler = TSharedPtr<FBuoyancyWaterSampler>(new FBuoyancyShallowWaterSampler(ShallowWaterSimData, ClosestWaterPointToParticle, WaterVelocity));
 	}
diff --git a/Engine/Plugins/Experimental/Buoyancy/Source/Runtime/Public/BuoyancyParticleData.h b/Engine/Plugins/Experimental/Buoyancy/Source/Runtime/Public/BuoyancyParticleData.h
--- a/Engine/Plugins/Experimental/Buoyancy/Source/Runtime/Public/BuoyancyParticleData.h
+++ b/Engine/Plugins/Experimental/Buoyancy/Source/Runtime/Public/BuoyancyParticleData.h
@@ -53,7 +53,7 @@ protected:
 class FBuoyancyShallowWaterSampler : public FBuoyancyWaterSampler
 {
 public:
-	FBuoyancyShallowWaterSampler(const FShallowWaterSimulationGrid& InWaterSampler, const FVector& InClosestWaterPointToParticle, const FVector &InAverageVelocity) :
+	FBuoyancyShallowWaterSampler(const TWeakObjectPtr<UShallowWaterSimulationData> InWaterSampler, const FVector& InClosestWaterPointToParticle, const FVector &InAverageVelocity) :
 		WaterSampler(InWaterSampler)
 	{
 		SamplerType = EWaterSamplerType::ShallowWater;
@@ -64,25 +64,45 @@ public:
 
 	virtual void SampleWaterAtPosition(const FVector& WorldPosition, FVector& OutVelocity, float& OutHeight, float& OutDepth) const override
 	{
-		WaterSampler.SampleShallowWaterSimulationAtPosition(WorldPosition, OutVelocity, OutHeight, OutDepth);
+		if (WaterSampler.IsValid() && WaterSampler->HasValidData())
+		{
+			WaterSampler->SampleShallowWaterSimulationAtPosition(WorldPosition, OutVelocity, OutHeight, OutDepth);
+		}
+		else
+		{
+			OutVelocity = {0.f,0.f,0.f};
+			OutHeight = 0.f;
+			OutDepth = 0.f;
+		}
 	}
 
 	virtual FVector SampleVelocityAtPosition(const FVector& WorldPosition) const override
 	{
-		FVector WaterVelocity;
-		float TmpWaterVal;
-		WaterSampler.SampleShallowWaterSimulationAtPosition(WorldPosition, WaterVelocity, TmpWaterVal, TmpWaterVal);
+		FVector WaterVelocity = {0,0,0};
+
+		if (WaterSampler.IsValid() && WaterSampler->HasValidData())
+		{
+			float TmpWaterVal;
+			WaterSampler->SampleShallowWaterSimulationAtPosition(WorldPosition, WaterVelocity, TmpWaterVal, TmpWaterVal);
+		}
 
 		return WaterVelocity;
 	}
 
 	virtual FVector SampleNormalAtPosition(const FVector& WorldPosition) const override
 	{
-		return WaterSampler.ComputeShallowWaterSimulationNormalAtPosition(WorldPosition);
+		if (WaterSampler.IsValid() && WaterSampler->HasValidData())
+		{
+			return WaterSampler->ComputeShallowWaterSimulationNormalAtPosition(WorldPosition);
+		}
+		else
+		{
+			return FVector(0,0,1);
+		}
 	}
 
 private:
-	const FShallowWaterSimulationGrid& WaterSampler;
+	const TWeakObjectPtr<UShallowWaterSimulationData> WaterSampler;
 };
 
 // sampler for spline rivers, but we are only doing 1 spline sample per interaction to 
diff --git a/Engine/Plugins/Experimental/Buoyancy/Source/Runtime/Public/BuoyancyWaterSplineData.h b/Engine/Plugins/Experimental/Buoyancy/Source/Runtime/Public/BuoyancyWaterSplineData.h
--- a/Engine/Plugins/Experimental/Buoyancy/Source/Runtime/Public/BuoyancyWaterSplineData.h
+++ b/Engine/Plugins/Experimental/Buoyancy/Source/Runtime/Public/BuoyancyWaterSplineData.h
@@ -6,14 +6,14 @@
 
 struct FBuoyancyWaterSplineData
 {
-	FBuoyancyWaterSplineData() { }
+	FBuoyancyWaterSplineData() {}
 	FBuoyancyWaterSplineData(
 		const Chaos::FRigidTransform3& InTransform,
 		const FInterpCurveVector& InPosition,
 		const EWaterBodyType InBodyType,
 		const TOptional<FInterpCurveFloat>& InWidth,
 		const TOptional<FInterpCurveFloat>& InVelocity,
-		const TOptional<FShallowWaterSimulationGrid>& InShallowWaterSimData
+		const TWeakObjectPtr<UShallowWaterSimulationData> InShallowWaterSimData
 	)
 		: Transform(InTransform)
 		, Position(InPosition)
@@ -26,7 +26,7 @@ struct FBuoyancyWaterSplineData
 
 	bool ShouldSampleFromShallowWaterSimulation() const
 	{
-		return ShallowWaterSimData.IsSet() && ShallowWaterSimData->IsValid();
+		return ShallowWaterSimData.IsValid() && ShallowWaterSimData->HasValidData();
 	}
 
 	// Parameters that all water bodies have
@@ -40,5 +40,5 @@ struct FBuoyancyWaterSplineData
 
 	// Only some water bodies can have a baked shallow water sim representation, which
 	// overrides the splines for water height/depth/velocity/normal
-	TOptional<FShallowWaterSimulationGrid> ShallowWaterSimData;
+	const TWeakObjectPtr<UShallowWaterSimulationData> ShallowWaterSimData;
 };
\ No newline at end of file
diff --git a/Engine/Plugins/Experimental/Water/Source/Runtime/Private/BakedShallowWaterSimulationComponent.cpp b/Engine/Plugins/Experimental/Water/Source/Runtime/Private/BakedShallowWaterSimulationComponent.cpp
--- a/Engine/Plugins/Experimental/Water/Source/Runtime/Private/BakedShallowWaterSimulationComponent.cpp
+++ b/Engine/Plugins/Experimental/Water/Source/Runtime/Private/BakedShallowWaterSimulationComponent.cpp
@@ -1,17 +1,25 @@
 // Copyright Epic Games, Inc. All Rights Reserved.
 #include "BakedShallowWaterSimulationComponent.h"
 #include "WaterBodyActor.h"
+#include "WaterModule.h"
+#include "WaterVersion.h"
 
 #include UE_INLINE_GENERATED_CPP_BY_NAME(BakedShallowWaterSimulationComponent)
 
-void FShallowWaterSimulationGrid::SampleShallowWaterSimulationAtIndex(const FVector2D &QueryFloatIndex, FVector& OutWaterVelocity, float& OutWaterHeight, float& OutWaterDepth) const
+UBakedShallowWaterSimulationComponent::UBakedShallowWaterSimulationComponent(const FObjectInitializer& ObjectInitializer)
+	: Super(ObjectInitializer)
+{}
+
+void UShallowWaterSimulationData::SampleShallowWaterSimulationAtIndex(const FVector2D& QueryFloatIndex, FVector& OutWaterVelocity, float& OutWaterHeight, float& OutWaterDepth) const
 {
 	OutWaterHeight = 0;
 	OutWaterDepth = 0;
 	OutWaterVelocity = FVector(0, 0, 0);
-
+	
 	if (QueryFloatIndex.X > 0 && QueryFloatIndex.X < NumCells.X - 1 && QueryFloatIndex.Y > 0 && QueryFloatIndex.Y < NumCells.Y - 1)
 	{
+		check(ArrayValues.Num() == NumCells.X * NumCells.Y);
+
 		const FIntVector2 BaseIndex = FIntVector2(QueryFloatIndex.X, QueryFloatIndex.Y);
 		const FVector2D LerpValue = QueryFloatIndex - FVector2D(BaseIndex.X, BaseIndex.Y);
 
@@ -31,7 +39,7 @@ void FShallowWaterSimulationGrid::SampleShallowWaterSimulationAtIndex(const FVec
 	}
 }
 
-void FShallowWaterSimulationGrid::QueryShallowWaterSimulationAtIndex(const FIntVector2 &QueryIndex, FVector& OutWaterVelocity, float& OutWaterHeight, float& OutWaterDepth) const
+void UShallowWaterSimulationData::QueryShallowWaterSimulationAtIndex(const FIntVector2& QueryIndex, FVector& OutWaterVelocity, float& OutWaterHeight, float& OutWaterDepth) const
 {
 	OutWaterHeight = 0;
 	OutWaterDepth = 0;
@@ -39,6 +47,8 @@ void FShallowWaterSimulationGrid::QueryShallowWaterSimulationAtIndex(const FIntV
 
 	if (QueryIndex.X >= 0 && QueryIndex.X < NumCells.X && QueryIndex.Y >= 0 && QueryIndex.Y < NumCells.Y)
 	{
+		check(ArrayValues.Num() == NumCells.X * NumCells.Y);
+
 		const FVector4 Sample = ArrayValues[QueryIndex.X + QueryIndex.Y * NumCells.X];
 
 		OutWaterHeight = Sample.X + Position.Z;
@@ -47,7 +57,7 @@ void FShallowWaterSimulationGrid::QueryShallowWaterSimulationAtIndex(const FIntV
 	}
 }
 
-FVector FShallowWaterSimulationGrid::ComputeShallowWaterSimulationNormalAtPosition(const FVector &QueryPos) const
+FVector UShallowWaterSimulationData::ComputeShallowWaterSimulationNormalAtPosition(const FVector& QueryPos) const
 {
 	const FVector2D FloatIndexPos = WorldToFloatIndex(QueryPos);
 
@@ -87,6 +97,27 @@ FVector FShallowWaterSimulationGrid::ComputeShallowWaterSimulationNormalAtPositi
 	return Normal;
 }
 
-UBakedShallowWaterSimulationComponent::UBakedShallowWaterSimulationComponent(const FObjectInitializer& ObjectInitializer)
-	: Super(ObjectInitializer)
-{}
+void UBakedShallowWaterSimulationComponent::PostLoad()
+{
+	Super::PostLoad();
+
+  #if WITH_EDITORONLY_DATA
+	if (GetLinkerCustomVersion(FWaterCustomVersion::GUID) <
+		FWaterCustomVersion::MigrateShallowWaterSimulationToUObject)
+	{
+		if (BakedSimulationData == nullptr && SimulationData_DEPRECATED.IsValid())
+		{
+			BakedSimulationData = NewObject<UShallowWaterSimulationData>(this);
+			BakedSimulationData->ArrayValues = MoveTemp(SimulationData_DEPRECATED.ArrayValues);
+			BakedSimulationData->NumCells = SimulationData_DEPRECATED.NumCells;
+			BakedSimulationData->Position = SimulationData_DEPRECATED.Position;
+			BakedSimulationData->Size = SimulationData_DEPRECATED.Size;
+			BakedSimulationData->BakedTexture = SimulationData_DEPRECATED.BakedTexture;
+
+			SimulationData_DEPRECATED = FShallowWaterSimulationGrid_DEPRECATED();
+
+			UE_LOG(LogWater, Warning, TEXT("Water simulation data was stored in deprecated format.  Converted to new format.  Recommended resave."));
+		}
+	}
+#endif
+}
diff --git a/Engine/Plugins/Experimental/Water/Source/Runtime/Private/WaterBodyComponent.cpp b/Engine/Plugins/Experimental/Water/Source/Runtime/Private/WaterBodyComponent.cpp
--- a/Engine/Plugins/Experimental/Water/Source/Runtime/Private/WaterBodyComponent.cpp
+++ b/Engine/Plugins/Experimental/Water/Source/Runtime/Private/WaterBodyComponent.cpp
@@ -667,7 +667,7 @@ FWaterBodyQueryResult UWaterBodyComponent::QueryWaterInfoClosestToWorldLocation(
 
 	if (bUseShallowWaterValues)
 	{		
-		BakedShallowWaterSim->SimulationData.SampleShallowWaterSimulationAtPosition(InWorldLocation,
+		BakedShallowWaterSim->BakedSimulationData->SampleShallowWaterSimulationAtPosition(InWorldLocation,
 			ShallowWaterVelocity, ShallowWaterHeight, ShallowWaterDepth);
 	}
 
@@ -719,7 +719,7 @@ FWaterBodyQueryResult UWaterBodyComponent::QueryWaterInfoClosestToWorldLocation(
 		{
 			if (bUseShallowWaterValues)
 			{
-				WaterPlaneNormal = BakedShallowWaterSim->SimulationData.ComputeShallowWaterSimulationNormalAtPosition(InWorldLocation);
+				WaterPlaneNormal = BakedShallowWaterSim->BakedSimulationData->ComputeShallowWaterSimulationNormalAtPosition(InWorldLocation);
 			}
 			else
 			{
diff --git a/Engine/Plugins/Experimental/Water/Source/Runtime/Public/BakedShallowWaterSimulationComponent.h b/Engine/Plugins/Experimental/Water/Source/Runtime/Public/BakedShallowWaterSimulationComponent.h
--- a/Engine/Plugins/Experimental/Water/Source/Runtime/Public/BakedShallowWaterSimulationComponent.h
+++ b/Engine/Plugins/Experimental/Water/Source/Runtime/Public/BakedShallowWaterSimulationComponent.h
@@ -11,43 +11,45 @@
 
 class AWaterBody;
 
-USTRUCT(BlueprintType)
-struct FShallowWaterSimulationGrid
+UCLASS(MinimalAPI)
+class UShallowWaterSimulationData : public UObject 
 {
 	GENERATED_BODY()
-
 public:
-	FShallowWaterSimulationGrid(const TArray<FVector4> &InArrayValues, TObjectPtr<UTexture2D> InBakedWaterTexture, const FIntVector2 &InNumCells, const FVector &InPosition, const FVector2D &InSize)
-		: ArrayValues(InArrayValues), 
-		NumCells(InNumCells), 
-		Position(InPosition), 
-		Size(InSize),
-		BakedTexture(InBakedWaterTexture)
-	{
-		Dx = Size / FVector2D(NumCells.X, NumCells.Y);
-	}
-	FShallowWaterSimulationGrid()
-	{				
-		NumCells = { 0, 0 };
-		Position = { 0.0, 0.0, 0.0 };		
-		Size = { 0.0, 0.0 };
-		Dx = { 0.0, 0.0 };
-	};
+	/** Raw grid data */
+	UPROPERTY()
+	TArray<FVector4> ArrayValues;
+	
+	/** Number of cells */
+	UPROPERTY()
+	FIntVector2 NumCells;
+	
+	/** World space center */
+	UPROPERTY()
+	FVector Position;
+	
+	/** World space size */
+	UPROPERTY()
+	FVector2D Size;
+	
+	/** Texture for baked water sim data */
+	UPROPERTY()
+	TObjectPtr<UTexture2D> BakedTexture;
 
 	/** Check for a valid sim domain */
-	bool IsValid() const
+	bool HasValidData() const
 	{
 		return NumCells.X > 0 && NumCells.Y > 0 && ArrayValues.Num() == NumCells.X * NumCells.Y && Size.X > SMALL_NUMBER && Size.Y > SMALL_NUMBER;
 	}
 
 	/** Linear interpolate sim data at an index */
-	UE_API void SampleShallowWaterSimulationAtIndex(const FVector2D &QueryFloatIndex, FVector& OutWaterVelocity, float& OutWaterHeight, float& OutWaterDepth) const;
+	WATER_API void SampleShallowWaterSimulationAtIndex(const FVector2D &QueryFloatIndex, FVector& OutWaterVelocity, float& OutWaterHeight, float& OutWaterDepth) const;
 	
 	/** Query sim data at an index */
-	UE_API void QueryShallowWaterSimulationAtIndex(const FIntVector2 &QueryIndex, FVector& OutWaterVelocity, float& OutWaterHeight, float& OutWaterDepth) const;
+	WATER_API void QueryShallowWaterSimulationAtIndex(const FIntVector2 &QueryIndex, FVector& OutWaterVelocity, float& OutWaterHeight, float& OutWaterDepth) const;
 
 	/** Compute normal at a world position */
-	UE_API FVector ComputeShallowWaterSimulationNormalAtPosition(const FVector &QueryPos) const;
+	WATER_API FVector ComputeShallowWaterSimulationNormalAtPosition(const FVector &QueryPos) const;
 	
 	/** Query shallow water grid cell values at a world position */
 	void QueryShallowWaterSimulationAtPosition(const FVector &QueryPos, FVector& OutWaterVelocity, float& OutWaterHeight, float& OutWaterDepth) const
@@ -73,16 +75,35 @@ public:
 	/** Convert world space to float index */
 	FVector2D WorldToFloatIndex(const FVector &WorldPos) const
 	{
+		if (Size.X < 1e-5 || Size.Y < 1e-5)
+		{
+			return FVector2D(0.f,0.f);
+		}
+
+
 		const FVector LocalPos = WorldPos - Position;
+		
+		// unit positions are 0 - 1, Local Pos has center of grid at origin
 		const FVector2D UnitPos = FVector2D(LocalPos.X, LocalPos.Y) / Size + .5;
+
+		// cell values are stored at grid cell center
 		return UnitPos * FVector2D(NumCells.X, NumCells.Y) - .5;
 	}
 
 	/** Convert a float index to world space */
 	FVector FloatIndexToWorld(const FVector2D &Index) const
 	{
+		if (NumCells.X == 0 || NumCells.Y == 0)
+		{
+			return FVector(0.f, 0.f, 0.f);
+		}
+
+		// Values are stored at cell center
 		const FVector2D UnitPos = (Index + .5) / FVector2D(NumCells.X, NumCells.Y);
+		
+		// local position origin is center of the grid
 		const FVector2D LocalPos = (UnitPos - .5) * Size;
+
 		const FVector LocalPos3D = FVector(LocalPos.X, LocalPos.Y, 0);
 		return LocalPos3D + Position;
 	}
@@ -92,6 +113,24 @@ public:
 	{
 		return FloatIndexToWorld(FVector2D(Index.X, Index.Y));
 	}
+};
+
+USTRUCT(BlueprintType)
+struct FShallowWaterSimulationGrid_DEPRECATED
+{
+	GENERATED_BODY()
+
+public:
+
+
+	/** Check for a valid sim domain */
+	bool IsValid() const
+	{
+		return NumCells.X > 0 && NumCells.Y > 0 && ArrayValues.Num() == NumCells.X * NumCells.Y && Size.X > SMALL_NUMBER && Size.Y > SMALL_NUMBER;
+	}
+
+	// All methods are removed since this is used for data migration only.  
+	// Any other usage should be updated to UShallowWaterSimulationData
 
 	/** Raw grid data */
 	UPROPERTY()
@@ -124,11 +163,18 @@ class UBakedShallowWaterSimulationComponent : public UPrimitiveComponent
 	GENERATED_UCLASS_BODY()
 
 public:
+ #if WITH_EDITORONLY_DATA
 	UPROPERTY()
-	FShallowWaterSimulationGrid SimulationData;
+	FShallowWaterSimulationGrid_DEPRECATED SimulationData_DEPRECATED;
+#endif
+
+	UPROPERTY()
+	TObjectPtr<UShallowWaterSimulationData> BakedSimulationData;
 
 	UPROPERTY()
 	TSet<TSoftObjectPtr<AWaterBody>> WaterBodies;
+
+	void PostLoad() override;
 };
 
 #undef UE_API
diff --git a/Engine/Plugins/Experimental/Water/Source/Runtime/Public/WaterBodyComponent.h b/Engine/Plugins/Experimental/Water/Source/Runtime/Public/WaterBodyComponent.h
--- a/Engine/Plugins/Experimental/Water/Source/Runtime/Public/WaterBodyComponent.h
+++ b/Engine/Plugins/Experimental/Water/Source/Runtime/Public/WaterBodyComponent.h
@@ -444,7 +444,7 @@ public:
 	void SetUseBakedSimulationForQueriesAndPhysics(bool bUseBakedSimulation) { bUseBakedSimForQueriesAndPhysics = bUseBakedSimulation;  }
 	
 	/** Query for if the baked simulations is valid for use */
-	bool UseBakedSimulationForQueriesAndPhysics() const { return bUseBakedSimForQueriesAndPhysics && BakedShallowWaterSim.IsValid() && BakedShallowWaterSim->SimulationData.IsValid(); }
+	bool UseBakedSimulationForQueriesAndPhysics() const { return bUseBakedSimForQueriesAndPhysics && BakedShallowWaterSim.IsValid() && BakedShallowWaterSim->BakedSimulationData != nullptr && BakedShallowWaterSim->BakedSimulationData->HasValidData(); }
 
 	/** 
 	 * Marks the owning water zone for rebuild. 
@@ -677,8 +677,7 @@ protected:
 	UPROPERTY(Transient)
 	mutable TWeakObjectPtr<ALandscapeProxy> Landscape;
 
-	/** Internal - Automatically points to the detected owning water zone. Should not be edited manually. */
-	UPROPERTY(Category = Water, EditAnywhere, AdvancedDisplay, TextExportTransient)
+	UPROPERTY(Category = Water, VisibleInstanceOnly, AdvancedDisplay, TextExportTransient)
 	TSoftObjectPtr<AWaterZone> OwningWaterZone;
 
 	UPROPERTY(Category = Water, EditAnywhere, BlueprintReadOnly, AdvancedDisplay)
diff --git a/Engine/Plugins/Experimental/Water/Source/Runtime/Public/WaterVersion.h b/Engine/Plugins/Experimental/Water/Source/Runtime/Public/WaterVersion.h
--- a/Engine/Plugins/Experimental/Water/Source/Runtime/Public/WaterVersion.h
+++ b/Engine/Plugins/Experimental/Water/Source/Runtime/Public/WaterVersion.h
@@ -32,7 +32,8 @@ struct FWaterCustomVersion
 		DefaultWaterInfoMaterialAssignmentMoved,
 		// Rebuild water render data to ensure no meshes from the broken tset change were serialized
 		RebuildWaterMeshDataTSetChange,
-
+		// Moved baked simulation data over to UObject for better memory usage and asset management
+		MigrateShallowWaterSimulationToUObject,
 		// -----<new versions can be added above this line>-------------------------------------------------
 		VersionPlusOne,
 		LatestVersion = VersionPlusOne - 1
diff --git a/Engine/Plugins/Experimental/WaterAdvanced/Source/WaterAdvanced/Private/ShallowWaterRiverActor.cpp b/Engine/Plugins/Experimental/WaterAdvanced/Source/WaterAdvanced/Private/ShallowWaterRiverActor.cpp
--- a/Engine/Plugins/Experimental/WaterAdvanced/Source/WaterAdvanced/Private/ShallowWaterRiverActor.cpp
+++ b/Engine/Plugins/Experimental/WaterAdvanced/Source/WaterAdvanced/Private/ShallowWaterRiverActor.cpp
@@ -202,6 +202,7 @@ void UShallowWaterRiverComponent::PostLoad()
 {
 	Super::PostLoad();	
 
+
 	// ensure all baked textures are virtual for backwards compatibility
 	ConvertToVirtualTextures();
 
@@ -945,7 +946,15 @@ void UShallowWaterRiverComponent::Bake()
 	}
 
 	BakedSim = NewObject<UBakedShallowWaterSimulationComponent>(this, NAME_None, RF_Public);
-	BakedSim->SimulationData = FShallowWaterSimulationGrid(ShallowWaterSimArrayValues, BakedWaterSurfaceTexture, FIntVector2(BakedWaterSurfaceRT->SizeX, BakedWaterSurfaceRT->SizeY), SystemPos, WorldGridSize);
+	
+	BakedSim->BakedSimulationData = NewObject<UShallowWaterSimulationData>(BakedSim, NAME_None, RF_Public);	
+	BakedSim->BakedSimulationData->ArrayValues = ShallowWaterSimArrayValues;
+	BakedSim->BakedSimulationData->NumCells = FIntVector2(BakedWaterSurfaceRT->SizeX, BakedWaterSurfaceRT->SizeY);
+	BakedSim->BakedSimulationData->Position = SystemPos;
+	BakedSim->BakedSimulationData->Size = WorldGridSize;
+	BakedSim->BakedSimulationData->BakedTexture = BakedWaterSurfaceTexture;
+
+
 	BakedSim->WaterBodies = AllWaterBodies;	
 		
 	// compute the maximum water height for each convex in each water body simulated by this river
@@ -953,11 +962,11 @@ void UShallowWaterRiverComponent::Bake()
 	TMap<FKConvexElem*, float> ConvexToMaxHeight;
 	for (int32 y = 0; y < BakedWaterSurfaceRT->SizeY; ++y) {
 	for (int32 x = 0; x < BakedWaterSurfaceRT->SizeX; ++x) {
-		FVector WorldPos = BakedSim->SimulationData.IndexToWorld(FIntVector2(x, y));
+		FVector WorldPos = BakedSim->BakedSimulationData->IndexToWorld(FIntVector2(x, y));
 
 		FVector Vel;
 		float Height, Depth;
-		BakedSim->SimulationData.QueryShallowWaterSimulationAtIndex(FIntVector2(x, y), Vel, Height, Depth);
+		BakedSim->BakedSimulationData->QueryShallowWaterSimulationAtIndex(FIntVector2(x, y), Vel, Height, Depth);
 		WorldPos.Z = Height;
 				
 		if (Depth > 1e-5)
@@ -969,29 +978,36 @@ void UShallowWaterRiverComponent::Bake()
 				TArray<UPrimitiveComponent*> CollisionComponents = CurrWaterBodyComponent->GetCollisionComponents();
 				for (UPrimitiveComponent* CurrCollisionComponent : CollisionComponents)
 				{
-					USplineMeshComponent* CurrSplineComponent = StaticCast<USplineMeshComponent*>(CurrCollisionComponent);
-					TObjectPtr<UBodySetup> CurrBodySetup = CurrSplineComponent->BodySetup;
-
-					const FTransform CurrMeshTransform = CurrCollisionComponent->GetComponentTransform();
+					USplineMeshComponent* CurrSplineComponent = Cast<USplineMeshComponent>(CurrCollisionComponent);
+					
+					if (CurrSplineComponent)
+					{
+						TObjectPtr<UBodySetup> CurrBodySetup = CurrSplineComponent->BodySetup;
 
-					// make sure the collision convex hull vertices are clamped to the min/max water height
-					for (FKConvexElem& ConvexElem : CurrBodySetup->AggGeom.ConvexElems)
-					{					
-						const TArray<FVector>& VertexData = ConvexElem.VertexData;		
-
-						// see if the current point is inside the convex projected to the xy plane
-						const FBox CurrBox = ConvexElem.CalcAABB(CurrMeshTransform, FVector(1, 1, 1));										
-
-						if (CurrBox.IsInsideXY(FBox(WorldPos, WorldPos)))
-						{						
-							float* TmpMaxHeight = ConvexToMaxHeight.Find(&ConvexElem);
-							if (TmpMaxHeight == nullptr)
-							{							
-								ConvexToMaxHeight.Emplace(&ConvexElem, WorldPos.Z);
-							}
-							else
-							{
-								*TmpMaxHeight = FMath::Max(*TmpMaxHeight, WorldPos.Z);
+						if (CurrBodySetup)
+						{
+							const FTransform CurrMeshTransform = CurrCollisionComponent->GetComponentTransform();
+
+							// make sure the collision convex hull vertices are clamped to the min/max water height
+							for (FKConvexElem& ConvexElem : CurrBodySetup->AggGeom.ConvexElems)
+							{					
+								const TArray<FVector>& VertexData = ConvexElem.VertexData;		
+
+								// see if the current point is inside the convex projected to the xy plane
+								const FBox CurrBox = ConvexElem.CalcAABB(CurrMeshTransform, FVector(1, 1, 1));										
+
+								if (CurrBox.IsInsideXY(FBox(WorldPos, WorldPos)))
+								{						
+									float* TmpMaxHeight = ConvexToMaxHeight.Find(&ConvexElem);
+									if (TmpMaxHeight == nullptr)
+									{							
+										ConvexToMaxHeight.Emplace(&ConvexElem, WorldPos.Z);
+									}
+									else
+									{
+										*TmpMaxHeight = FMath::Max(*TmpMaxHeight, WorldPos.Z);
+									}
+								}
 							}
 						}
 					}
@@ -1014,79 +1030,85 @@ void UShallowWaterRiverComponent::Bake()
 		// we will miss collisions.
 		for (UPrimitiveComponent* CurrCollisionComponent : CollisionComponents)
 		{
-			USplineMeshComponent* CurrSplineComponent = StaticCast<USplineMeshComponent*>(CurrCollisionComponent);
+			USplineMeshComponent* CurrSplineComponent = Cast<USplineMeshComponent>(CurrCollisionComponent);
 			
-			TObjectPtr<UBodySetup> CurrBodySetup = CurrSplineComponent->BodySetup;
-
-			FTransform CurrMeshTransform = CurrCollisionComponent->GetComponentTransform();
-
-			// make sure the collision convex hull vertices are clamped to the min/max water height
-			for (FKConvexElem& ConvexElem : CurrBodySetup->AggGeom.ConvexElems)
+			if (CurrSplineComponent)
 			{
-				TArray<FVector>& VertexData = ConvexElem.VertexData;
+				TObjectPtr<UBodySetup> CurrBodySetup = CurrSplineComponent->BodySetup;
 
-				if (const float* WorldMaxZForConvex = ConvexToMaxHeight.Find(&ConvexElem))
+				if (CurrBodySetup)
 				{
-					// for each vertex in the convex hull, set the Z to the maximum baked water sim Z height for the convex
-					int32 Idx = 0;
-					for (FVector& Vertex : VertexData)
-					{
-						// only top vertices are 4,5,6,7
-						if (Idx >= 4)
-						{
-							FVector VWorld = CurrMeshTransform.TransformPosition(Vertex);
+					FTransform CurrMeshTransform = CurrCollisionComponent->GetComponentTransform();
 
-							VWorld.Z = FMath::Max(VWorld.Z, *WorldMaxZForConvex);
-
-							const FVector VLocal = CurrMeshTransform.InverseTransformPosition(VWorld);
-							Vertex.X = VLocal.X;
-							Vertex.Y = VLocal.Y;
-							Vertex.Z = VLocal.Z;
-						}
-
-
-						#if ENABLE_DRAW_DEBUG
-						if (bShallowWaterRiverDebugVisualize)
-						{		
-							FVector VWorld = CurrMeshTransform.TransformPosition(Vertex);
+					// make sure the collision convex hull vertices are clamped to the min/max water height
+					for (FKConvexElem& ConvexElem : CurrBodySetup->AggGeom.ConvexElems)
+					{
+						TArray<FVector>& VertexData = ConvexElem.VertexData;
 
-							switch (Idx)
+						if (const float* WorldMaxZForConvex = ConvexToMaxHeight.Find(&ConvexElem))
+						{
+							// for each vertex in the convex hull, set the Z to the maximum baked water sim Z height for the convex
+							int32 Idx = 0;
+							for (FVector& Vertex : VertexData)
 							{
-								case 0:
-								DrawDebugSphere(GetWorld() , VWorld, 10., 2, FColor::Red, true);		
-								break;
-								case 1:
-								DrawDebugSphere(GetWorld() , VWorld, 10., 3, FColor::Green, true);		
-								break;
-								case 2:
-								DrawDebugSphere(GetWorld() , VWorld, 10., 4, FColor::Blue, true);		
-								break;
-								case 3:
-								DrawDebugSphere(GetWorld() , VWorld, 10., 5, FColor::Black, true);		
-								break;
-								case 4:
-								DrawDebugSphere(GetWorld() , VWorld, 10., 6, FColor::White, true);	//	
-								break;
-								case 5:
-								DrawDebugSphere(GetWorld() , VWorld, 10., 7, FColor::Magenta, true); //		
-								break;
-								case 6:
-								DrawDebugSphere(GetWorld() , VWorld, 10., 8, FColor::Orange, true);	//	
-								break;
-								case 7:
-								DrawDebugSphere(GetWorld() , VWorld, 10., 9, FColor::Purple, true);	//	
-								break;								
-							}
+								// only top vertices are 4,5,6,7
+								if (Idx >= 4)
+								{
+									FVector VWorld = CurrMeshTransform.TransformPosition(Vertex);
+
+									VWorld.Z = FMath::Max(VWorld.Z, *WorldMaxZForConvex);
+
+									const FVector VLocal = CurrMeshTransform.InverseTransformPosition(VWorld);
+									Vertex.X = VLocal.X;
+									Vertex.Y = VLocal.Y;
+									Vertex.Z = VLocal.Z;
+								}
+
+
+								#if ENABLE_DRAW_DEBUG
+								if (bShallowWaterRiverDebugVisualize)
+								{		
+									FVector VWorld = CurrMeshTransform.TransformPosition(Vertex);
+
+									switch (Idx)
+									{
+										case 0:
+										DrawDebugSphere(GetWorld() , VWorld, 10., 2, FColor::Red, true);		
+										break;
+										case 1:
+										DrawDebugSphere(GetWorld() , VWorld, 10., 3, FColor::Green, true);		
+										break;
+										case 2:
+										DrawDebugSphere(GetWorld() , VWorld, 10., 4, FColor::Blue, true);		
+										break;
+										case 3:
+										DrawDebugSphere(GetWorld() , VWorld, 10., 5, FColor::Black, true);		
+										break;
+										case 4:
+										DrawDebugSphere(GetWorld() , VWorld, 10., 6, FColor::White, true);	//	
+										break;
+										case 5:
+										DrawDebugSphere(GetWorld() , VWorld, 10., 7, FColor::Magenta, true); //		
+										break;
+										case 6:
+										DrawDebugSphere(GetWorld() , VWorld, 10., 8, FColor::Orange, true);	//	
+										break;
+										case 7:
+										DrawDebugSphere(GetWorld() , VWorld, 10., 9, FColor::Purple, true);	//	
+										break;								
+									}
 							
-						}
-						#endif
+								}
+								#endif
 												
-						Idx++;
-					}
+								Idx++;
+							}
+						}
+					}			
+			
+					CurrSplineComponent->PostEditChange();
 				}
-			}			
-
-			CurrSplineComponent->PostEditChange();
+			}
 		}
 
 		CurrWaterBodyComponent->PostEditChange();
diff --git a/Engine/Plugins/Experimental/WaterAdvanced/Source/WaterAdvanced/Private/ShallowWaterSubsystem.cpp b/Engine/Plugins/Experimental/WaterAdvanced/Source/WaterAdvanced/Private/ShallowWaterSubsystem.cpp
--- a/Engine/Plugins/Experimental/WaterAdvanced/Source/WaterAdvanced/Private/ShallowWaterSubsystem.cpp
+++ b/Engine/Plugins/Experimental/WaterAdvanced/Source/WaterAdvanced/Private/ShallowWaterSubsystem.cpp
@@ -466,6 +466,12 @@ APawn* UShallowWaterSubsystem::GetTheMostRelevantPlayerPawn() const
 
 void UShallowWaterSubsystem::CreateRTs()
 {
+	if (Settings == nullptr)
+	{
+		UE_LOG(LogShallowWater, Warning, TEXT("ShallowWaterComponent: Null settings."))
+		return;
+	}
+
 	int32 Resolution = Settings->ShallowWaterSimParameters.ResolutionMaxAxis;
 
 	if (Resolution <= 0)
@@ -476,21 +482,45 @@ void UShallowWaterSubsystem::CreateRTs()
 
 	NormalRT = NewObject<UTextureRenderTarget2D>(this);
 	check(NormalRT);
-	NormalRT->RenderTargetFormat = ETextureRenderTargetFormat::RTF_RGB10A2;
+	NormalRT->RenderTargetFormat = ETextureRenderTargetFormat::RTF_RGBA16f;
 	NormalRT->ClearColor = FLinearColor(0.f,0.f,0.f,0.f);
 	NormalRT->bAutoGenerateMips = false;
 	NormalRT->bCanCreateUAV = true;    // Niagara RT DI requires UAV. Do we need to create it here?
 	NormalRT->InitAutoFormat(Resolution, Resolution);	
 	NormalRT->UpdateResourceImmediate(true);	
+
+	if (Settings->OutputVelocity)
+	{
+		VelocityRT = NewObject<UTextureRenderTarget2D>(this);
+		check(VelocityRT);
+		VelocityRT->RenderTargetFormat = ETextureRenderTargetFormat::RTF_RG16f;
+		VelocityRT->ClearColor = FLinearColor(0.f,0.f,0.f,0.f);
+		VelocityRT->bAutoGenerateMips = false;
+		VelocityRT->bCanCreateUAV = true;    // Niagara RT DI requires UAV. Do we need to create it here?
+		VelocityRT->InitAutoFormat(Resolution, Resolution);	
+		VelocityRT->UpdateResourceImmediate(true);
+	}
+	else
+	{
+		VelocityRT = nullptr;
+	}
 }
 
 void UShallowWaterSubsystem::InitializeParameters()
 {
+	if (Settings == nullptr)
+	{
+		UE_LOG(LogShallowWater, Warning, TEXT("ShallowWaterComponent: Null settings."))
+		return;
+	}
+
 	if (ShallowWaterNiagaraSimulation)
 	{
 		ShallowWaterNiagaraSimulation->SetVariableVec2(FName("WorldGridSize"), FVector2D(GetGridSize()));
 		ShallowWaterNiagaraSimulation->SetVariableInt(FName("ResolutionMaxAxis"), GetGridResolution());
-		ShallowWaterNiagaraSimulation->SetVariableTextureRenderTarget(FName("NormalRT"), NormalRT);		
+		ShallowWaterNiagaraSimulation->SetVariableTextureRenderTarget(Settings->NormalRTNiagaraName, NormalRT);
+
+		ShallowWaterNiagaraSimulation->SetVariableTextureRenderTarget(Settings->VelocityRTNiagaraName, VelocityRT);
 
 		ShallowWaterNiagaraSimulation->SetVariableBool(FName("UseBakedSim"), false);
 	}
@@ -591,14 +621,14 @@ void UShallowWaterSubsystem::UpdateGridMovement()
 			if (WaterBodyComp->UseBakedSimulationForQueriesAndPhysics())
 			{
 				UBakedShallowWaterSimulationComponent* BakedSim = WaterBodyComp->GetBakedShallowWaterSimulation();				
-				UTexture* BakedSimTex = Cast<UTexture>(BakedSim->SimulationData.BakedTexture.Get());
+				UTexture* BakedSimTex = Cast<UTexture>(BakedSim->BakedSimulationData->BakedTexture.Get());
 
 				if (BakedSimTex != nullptr)
 				{
-					ShallowWaterNiagaraSimulation->SetVariableVec3(FName("BakedWaterSimLocation"), BakedSim->SimulationData.Position);
-					ShallowWaterNiagaraSimulation->SetVariableVec2(FName("BakedWaterSimSize"), BakedSim->SimulationData.Size);
+					ShallowWaterNiagaraSimulation->SetVariableVec3(FName("BakedWaterSimLocation"), BakedSim->BakedSimulationData->Position);
+					ShallowWaterNiagaraSimulation->SetVariableVec2(FName("BakedWaterSimSize"), BakedSim->BakedSimulationData->Size);
 					ShallowWaterNiagaraSimulation->SetVariableVec2(FName("BakedWaterSimRes"), 
-						FVector2D(BakedSim->SimulationData.NumCells.X, BakedSim->SimulationData.NumCells.Y));
+						FVector2D(BakedSim->BakedSimulationData->NumCells.X, BakedSim->BakedSimulationData->NumCells.Y));
 					ShallowWaterNiagaraSimulation->SetVariableTexture(FName("BakedWaterSimTexture"), BakedSimTex);
 
 					UseBakedSim = true;				
@@ -725,16 +755,19 @@ void UShallowWaterSubsystem::TryUpdateWaterBodyMIDParameters(UWaterBodyComponent
 	if (UMaterialInstanceDynamic* WaterMID = WaterBodyComponent->GetWaterMaterialInstance(); WaterMID)
 	{
 		WaterMID->SetTextureParameterValue(Settings->NormalRTMaterialName, NormalRT);
+		WaterMID->SetTextureParameterValue(Settings->VelocityRTMaterialName, VelocityRT);
 		WaterMID->SetScalarParameterValue(FName(TEXT("DEV_UseNewShallowWaterSubsystem")), 1.f);   // Temp before replacing OG system
 	}
 	if (UMaterialInstanceDynamic* WaterMID = WaterBodyComponent->GetRiverToOceanTransitionMaterialInstance(); WaterMID)
 	{
 		WaterMID->SetTextureParameterValue(Settings->NormalRTMaterialName, NormalRT);
+		WaterMID->SetTextureParameterValue(Settings->VelocityRTMaterialName, VelocityRT);
 		WaterMID->SetScalarParameterValue(FName(TEXT("DEV_UseNewShallowWaterSubsystem")), 1.f);   // Temp before replacing OG system
 	}
 	if (UMaterialInstanceDynamic* WaterMID = WaterBodyComponent->GetRiverToLakeTransitionMaterialInstance(); WaterMID)
 	{
 		WaterMID->SetTextureParameterValue(Settings->NormalRTMaterialName, NormalRT);
+		WaterMID->SetTextureParameterValue(Settings->VelocityRTMaterialName, VelocityRT);
 		WaterMID->SetScalarParameterValue(FName(TEXT("DEV_UseNewShallowWaterSubsystem")), 1.f);   // Temp before replacing OG system
 	}
 }
diff --git a/Engine/Plugins/Experimental/WaterAdvanced/Source/WaterAdvanced/Public/ShallowWaterSettings.h b/Engine/Plugins/Experimental/WaterAdvanced/Source/WaterAdvanced/Public/ShallowWaterSettings.h
--- a/Engine/Plugins/Experimental/WaterAdvanced/Source/WaterAdvanced/Public/ShallowWaterSettings.h
+++ b/Engine/Plugins/Experimental/WaterAdvanced/Source/WaterAdvanced/Public/ShallowWaterSettings.h
@@ -40,8 +40,16 @@ public:
 	FName ResolutionMaxAxisMPCName = FName("FluidSimResolution");
 	UPROPERTY(Config, BlueprintReadOnly, EditAnywhere, Category="Shallow Water Simulation")
 	FName NormalRTMaterialName = FName("NormalAndHeight");
+	UPROPERTY(Config, BlueprintReadOnly, EditAnywhere, Category="Shallow Water Simulation")
+	FName VelocityRTMaterialName = FName("Velocity");
+	UPROPERTY(Config, BlueprintReadOnly, EditAnywhere, Category="Shallow Water Simulation")
+	FName NormalRTNiagaraName = FName("NormalRT");
+	UPROPERTY(Config, BlueprintReadOnly, EditAnywhere, Category="Shallow Water Simulation")
+	FName VelocityRTNiagaraName = FName("VelocityRT");
 	UPROPERTY(Config, BlueprintReadOnly, EditAnywhere, Category = "Shallow Water Simulation")
 	bool UseDefaultShallowWaterSubsystem = false;
+	UPROPERTY(Config, BlueprintReadOnly, EditAnywhere, Category = "Shallow Water Simulation")
+	bool OutputVelocity = false;
 
 	UPROPERTY(Config, BlueprintReadOnly, EditAnywhere, Category="Interaction")
 	int32 MaxActivePawnNum = 6;
diff --git a/Engine/Plugins/Experimental/WaterAdvanced/Source/WaterAdvanced/Public/ShallowWaterSubsystem.h b/Engine/Plugins/Experimental/WaterAdvanced/Source/WaterAdvanced/Public/ShallowWaterSubsystem.h
--- a/Engine/Plugins/Experimental/WaterAdvanced/Source/WaterAdvanced/Public/ShallowWaterSubsystem.h
+++ b/Engine/Plugins/Experimental/WaterAdvanced/Source/WaterAdvanced/Public/ShallowWaterSubsystem.h
@@ -240,6 +240,8 @@ protected:
 	TWeakObjectPtr<const UTextureRenderTarget2DArray> WaterInfoTexture;
 	UPROPERTY(BlueprintReadOnly, VisibleAnywhere, Category="Shallow Water")
 	TObjectPtr<UTextureRenderTarget2D> NormalRT;
+	UPROPERTY(BlueprintReadOnly, VisibleAnywhere, Category="Shallow Water")
+	TObjectPtr<UTextureRenderTarget2D> VelocityRT;
 
 	UPROPERTY(BlueprintReadOnly, VisibleAnywhere, Category="Shallow Water")
 	TObjectPtr<UShallowWaterSettings> Settings;
