speedup collision geometry computation when baking river

#Synced-CL 51597297

diff --git a/Engine/Plugins/Experimental/Water/Source/Runtime/Public/BakedShallowWaterSimulationComponent.h b/Engine/Plugins/Experimental/Water/Source/Runtime/Public/BakedShallowWaterSimulationComponent.h
--- a/Engine/Plugins/Experimental/Water/Source/Runtime/Public/BakedShallowWaterSimulationComponent.h
+++ b/Engine/Plugins/Experimental/Water/Source/Runtime/Public/BakedShallowWaterSimulationComponent.h
@@ -216,7 +216,10 @@ class UShallowWaterSimulationDataSparse : public UShallowWaterSimulationDataBase
 	GENERATED_BODY()
 
 public:
-		
+
+	/** Block size for sparse storage (16x16 cells per block) */
+	static constexpr int32 BlockSize = 16;
+
 	/** Raw grid data */
 	UPROPERTY()
 	TArray<uint16> BlockValues;
@@ -297,7 +300,6 @@ public:
 		}}
 	}
 private:
-	static constexpr int32 BlockSize = 16;
 	static constexpr int32 DataSizeInBytes = 4;
 	static constexpr int32 BlockSizeInBytes = BlockSize * BlockSize * 4;
 };
diff --git a/Engine/Plugins/Experimental/WaterAdvanced/Source/WaterAdvanced/Private/ShallowWaterRiverActor.cpp b/Engine/Plugins/Experimental/WaterAdvanced/Source/WaterAdvanced/Private/ShallowWaterRiverActor.cpp
--- a/Engine/Plugins/Experimental/WaterAdvanced/Source/WaterAdvanced/Private/ShallowWaterRiverActor.cpp
+++ b/Engine/Plugins/Experimental/WaterAdvanced/Source/WaterAdvanced/Private/ShallowWaterRiverActor.cpp
@@ -941,163 +941,176 @@ void UShallowWaterRiverComponent::Bake()
 	BakedSimulationData->Size = WorldGridSize;
 	BakedSimulationData->BakedTexture = BakedWaterSurfaceTexture;
 
-	BakedSim->WaterBodies = AllWaterBodies;	
-			
-	// compute the maximum water height for each convex in each water body simulated by this river
-	// we use this to modify the collision geometry so it fully encompasses the baked water sim
-	TMap<FKConvexElem*, float> ConvexToMaxHeight;
-	for (int32 y = 0; y < BakedWaterSurfaceRT->SizeY; ++y) {
-	for (int32 x = 0; x < BakedWaterSurfaceRT->SizeX; ++x) {
-		FVector WorldPos = BakedSim->BakedSimulationData->IndexToWorld(FIntVector2(x, y));
-
-		FVector Vel;
-		float Height, Depth;
-		BakedSim->BakedSimulationData->QueryShallowWaterSimulationAtIndex(FIntVector2(x, y), Vel, Height, Depth);
-		WorldPos.Z = Height;
-				
-		if (Depth > 1e-5)
-		{
-			for (TSoftObjectPtr<AWaterBody > CurrWaterBody : AllWaterBodies)
-			{						
-				TObjectPtr<UWaterBodyComponent> CurrWaterBodyComponent = CurrWaterBody->GetWaterBodyComponent();
+	BakedSim->WaterBodies = AllWaterBodies;
+
+	// Compute the maximum water height for each convex in each water body simulated by this river.
+	// We use this to modify the collision geometry so it fully encompasses the baked water sim.
+
+	// Gather all convex elements once to avoid redundant lookups per cell
+	struct FConvexInfo
+	{
+		FKConvexElem* Convex;
+		FTransform Transform;
+		FBox AABBox;
+		USplineMeshComponent* SplineComponent; // Needed for PostEditChange after vertex modification
+	};
+	TArray<FConvexInfo> AllConvexes;
 
-				TArray<UPrimitiveComponent*> CollisionComponents = CurrWaterBodyComponent->GetCollisionComponents();
-				for (UPrimitiveComponent* CurrCollisionComponent : CollisionComponents)
+	for (TSoftObjectPtr<AWaterBody> CurrWaterBody : AllWaterBodies)
+	{
+		if (TObjectPtr<UWaterBodyComponent> CurrWaterBodyComponent = CurrWaterBody->GetWaterBodyComponent())
+		{		
+			TArray<UPrimitiveComponent*> CollisionComponents = CurrWaterBodyComponent->GetCollisionComponents();
+
+			for (UPrimitiveComponent* CurrCollisionComponent : CollisionComponents)
+			{
+				if (USplineMeshComponent* CurrSplineComponent = Cast<USplineMeshComponent>(CurrCollisionComponent))
 				{
-					USplineMeshComponent* CurrSplineComponent = Cast<USplineMeshComponent>(CurrCollisionComponent);
-					
-					if (CurrSplineComponent)
+					if (TObjectPtr<UBodySetup> CurrBodySetup = CurrSplineComponent->BodySetup)
 					{
-						TObjectPtr<UBodySetup> CurrBodySetup = CurrSplineComponent->BodySetup;
-
-						if (CurrBodySetup)
+						const FTransform CurrMeshTransform = CurrCollisionComponent->GetComponentTransform();
+						for (FKConvexElem& ConvexElem : CurrBodySetup->AggGeom.ConvexElems)
 						{
-							const FTransform CurrMeshTransform = CurrCollisionComponent->GetComponentTransform();
-
-							// make sure the collision convex hull vertices are clamped to the min/max water height
-							for (FKConvexElem& ConvexElem : CurrBodySetup->AggGeom.ConvexElems)
-							{					
-								const TArray<FVector>& VertexData = ConvexElem.VertexData;		
-
-								// see if the current point is inside the convex projected to the xy plane
-								const FBox CurrBox = ConvexElem.CalcAABB(CurrMeshTransform, FVector(1, 1, 1));										
-
-								if (CurrBox.IsInsideXY(FBox(WorldPos, WorldPos)))
-								{						
-									float* TmpMaxHeight = ConvexToMaxHeight.Find(&ConvexElem);
-									if (TmpMaxHeight == nullptr)
-									{							
-										ConvexToMaxHeight.Emplace(&ConvexElem, WorldPos.Z);
-									}
-									else
-									{
-										*TmpMaxHeight = FMath::Max(*TmpMaxHeight, WorldPos.Z);
-									}
-								}
-							}
+							const FBox CurrBox = ConvexElem.CalcAABB(CurrMeshTransform, FVector(1, 1, 1));
+							AllConvexes.Add({ &ConvexElem, CurrMeshTransform, CurrBox, CurrSplineComponent });
 						}
 					}
 				}
 			}
 		}
-	}}
+	}
 
-	// set the sim texture on each water body that is in the simulated river.  
-	for (TSoftObjectPtr<AWaterBody > CurrWaterBody : AllWaterBodies)
+	// Iterate over allocated blocks only (sparse storage is always used - see line 929)
+	
+	// Map each convex to the maximum water height found within its bounds
+	TMap<FKConvexElem*, float> ConvexToMaxHeight;
+
+	for (const auto& BlockPair : BakedSimulationData->GridIndexToBlockIndex)
 	{
-		TObjectPtr<UWaterBodyComponent> CurrWaterBodyComponent = CurrWaterBody->GetWaterBodyComponent();
-				
-		CurrWaterBodyComponent->SetBakedShallowWaterSimulation(BakedSim);
+		const FIntVector2 BlockIndex = BlockPair.Key;
+
+		// Compute cell range for this block
+		const int32 StartX = BlockIndex.X * UShallowWaterSimulationDataSparse::BlockSize;
+		const int32 StartY = BlockIndex.Y * UShallowWaterSimulationDataSparse::BlockSize;
+		const int32 EndX = FMath::Min(StartX + UShallowWaterSimulationDataSparse::BlockSize, BakedWaterSurfaceRT->SizeX);
+		const int32 EndY = FMath::Min(StartY + UShallowWaterSimulationDataSparse::BlockSize, BakedWaterSurfaceRT->SizeY);
 
-		// grow bounds in z to include the tallest height
-		TArray<UPrimitiveComponent*> CollisionComponents = CurrWaterBodyComponent->GetCollisionComponents();
+		// Process only cells within this allocated block
+		for (int32 y = StartY; y < EndY; ++y) {
+		for (int32 x = StartX; x < EndX; ++x) {
+			FVector WorldPos = BakedSimulationData->IndexToWorld(FIntVector2(x, y));
 
-		// Make sure that the collision objects includes the maximum height of the baked water sim otherwise
-		// we will miss collisions.
-		for (UPrimitiveComponent* CurrCollisionComponent : CollisionComponents)
+			FVector Vel;
+			float Height, Depth;
+			BakedSimulationData->QueryShallowWaterSimulationAtIndex(FIntVector2(x, y), Vel, Height, Depth);
+			WorldPos.Z = Height;
+
+			// Depth check still needed since sparse blocks can have some zero-depth cells
+			if (Depth > 1e-5)
+			{
+				// Test against pre-gathered convex elements
+				for (const FConvexInfo& ConvexInfo : AllConvexes)
+				{
+					// See if the current point is inside the convex projected to the xy plane
+					if (ConvexInfo.AABBox.IsInsideXY(FBox(WorldPos, WorldPos)))
+					{
+						float* TmpMaxHeight = ConvexToMaxHeight.Find(ConvexInfo.Convex);
+						if (TmpMaxHeight == nullptr)
+						{
+							ConvexToMaxHeight.Emplace(ConvexInfo.Convex, WorldPos.Z);
+						}
+						else
+						{
+							*TmpMaxHeight = FMath::Max(*TmpMaxHeight, WorldPos.Z);
+						}
+					}
+				}
+			}
+		}}
+	}
+
+	// Apply the maximum water heights to convex collision geometry
+	// Use AllConvexes list to avoid redundant component traversal
+	TSet<USplineMeshComponent*> ModifiedSplineComponents;
+	for (const FConvexInfo& ConvexInfo : AllConvexes)
+	{
+		if (const float* WorldMaxZForConvex = ConvexToMaxHeight.Find(ConvexInfo.Convex))
 		{
-			USplineMeshComponent* CurrSplineComponent = Cast<USplineMeshComponent>(CurrCollisionComponent);
-			
-			if (CurrSplineComponent)
+			TArray<FVector>& VertexData = ConvexInfo.Convex->VertexData;
+
+			// For each vertex in the convex hull, set the Z to the maximum baked water sim Z height
+			int32 Idx = 0;
+			for (FVector& Vertex : VertexData)
 			{
-				TObjectPtr<UBodySetup> CurrBodySetup = CurrSplineComponent->BodySetup;
+				// Only top vertices are 4,5,6,7
+				if (Idx >= 4)
+				{
+					FVector VWorld = ConvexInfo.Transform.TransformPosition(Vertex);
+					VWorld.Z = FMath::Max(VWorld.Z, *WorldMaxZForConvex);
+
+					const FVector VLocal = ConvexInfo.Transform.InverseTransformPosition(VWorld);
+					Vertex.X = VLocal.X;
+					Vertex.Y = VLocal.Y;
+					Vertex.Z = VLocal.Z;
+				}
 
-				if (CurrBodySetup)
+				#if ENABLE_DRAW_DEBUG
+				if (bShallowWaterRiverDebugVisualize)
 				{
-					FTransform CurrMeshTransform = CurrCollisionComponent->GetComponentTransform();
+					FVector VWorld = ConvexInfo.Transform.TransformPosition(Vertex);
 
-					// make sure the collision convex hull vertices are clamped to the min/max water height
-					for (FKConvexElem& ConvexElem : CurrBodySetup->AggGeom.ConvexElems)
+					switch (Idx)
 					{
-						TArray<FVector>& VertexData = ConvexElem.VertexData;
-
-						if (const float* WorldMaxZForConvex = ConvexToMaxHeight.Find(&ConvexElem))
-						{
-							// for each vertex in the convex hull, set the Z to the maximum baked water sim Z height for the convex
-							int32 Idx = 0;
-							for (FVector& Vertex : VertexData)
-							{
-								// only top vertices are 4,5,6,7
-								if (Idx >= 4)
-								{
-									FVector VWorld = CurrMeshTransform.TransformPosition(Vertex);
-
-									VWorld.Z = FMath::Max(VWorld.Z, *WorldMaxZForConvex);
-
-									const FVector VLocal = CurrMeshTransform.InverseTransformPosition(VWorld);
-									Vertex.X = VLocal.X;
-									Vertex.Y = VLocal.Y;
-									Vertex.Z = VLocal.Z;
-								}
-
-
-								#if ENABLE_DRAW_DEBUG
-								if (bShallowWaterRiverDebugVisualize)
-								{		
-									FVector VWorld = CurrMeshTransform.TransformPosition(Vertex);
-
-									switch (Idx)
-									{
-										case 0:
-										DrawDebugSphere(GetWorld() , VWorld, 10., 2, FColor::Red, true);		
-										break;
-										case 1:
-										DrawDebugSphere(GetWorld() , VWorld, 10., 3, FColor::Green, true);		
-										break;
-										case 2:
-										DrawDebugSphere(GetWorld() , VWorld, 10., 4, FColor::Blue, true);		
-										break;
-										case 3:
-										DrawDebugSphere(GetWorld() , VWorld, 10., 5, FColor::Black, true);		
-										break;
-										case 4:
-										DrawDebugSphere(GetWorld() , VWorld, 10., 6, FColor::White, true);	//	
-										break;
-										case 5:
-										DrawDebugSphere(GetWorld() , VWorld, 10., 7, FColor::Magenta, true); //		
-										break;
-										case 6:
-										DrawDebugSphere(GetWorld() , VWorld, 10., 8, FColor::Orange, true);	//	
-										break;
-										case 7:
-										DrawDebugSphere(GetWorld() , VWorld, 10., 9, FColor::Purple, true);	//	
-										break;								
-									}
-							
-								}
-								#endif
-												
-								Idx++;
-							}
-						}
-					}			
-			
-					CurrSplineComponent->PostEditChange();
+						case 0:
+						DrawDebugSphere(GetWorld(), VWorld, 10., 2, FColor::Red, true);
+						break;
+						case 1:
+						DrawDebugSphere(GetWorld(), VWorld, 10., 3, FColor::Green, true);
+						break;
+						case 2:
+						DrawDebugSphere(GetWorld(), VWorld, 10., 4, FColor::Blue, true);
+						break;
+						case 3:
+						DrawDebugSphere(GetWorld(), VWorld, 10., 5, FColor::Black, true);
+						break;
+						case 4:
+						DrawDebugSphere(GetWorld(), VWorld, 10., 6, FColor::White, true);
+						break;
+						case 5:
+						DrawDebugSphere(GetWorld(), VWorld, 10., 7, FColor::Magenta, true);
+						break;
+						case 6:
+						DrawDebugSphere(GetWorld(), VWorld, 10., 8, FColor::Orange, true);
+						break;
+						case 7:
+						DrawDebugSphere(GetWorld(), VWorld, 10., 9, FColor::Purple, true);
+						break;
+					}
 				}
+				#endif
+
+				Idx++;
 			}
+
+			// Track modified components for PostEditChange
+			ModifiedSplineComponents.Add(ConvexInfo.SplineComponent);
 		}
+	}
 
-		CurrWaterBodyComponent->PostEditChange();
+	// Notify modified spline components
+	for (USplineMeshComponent* SplineComponent : ModifiedSplineComponents)
+	{
+		SplineComponent->PostEditChange();
+	}
+
+	// Set the sim texture on each water body that is in the simulated river
+	for (TSoftObjectPtr<AWaterBody> CurrWaterBody : AllWaterBodies)
+	{
+		if (TObjectPtr<UWaterBodyComponent> CurrWaterBodyComponent = CurrWaterBody->GetWaterBodyComponent())		
+		{
+			CurrWaterBodyComponent->SetBakedShallowWaterSimulation(BakedSim);
+			CurrWaterBodyComponent->PostEditChange();
+		}
 	}
 }
 
